Move UAF Test plugin out of Restricted/NotForLicensees directory

diff --git a/Engine/Plugins/Experimental/UAF/UAFTests/Source/UAFTests/Private/UAFAnimationGraphTests.cpp b/Engine/Plugins/Experimental/UAF/UAFTests/Source/UAFTests/Private/UAFAnimationGraphTests.cpp
new file mode 100644
--- /dev/null
+++ b/Engine/Plugins/Experimental/UAF/UAFTests/Source/UAFTests/Private/UAFAnimationGraphTests.cpp
@@ -0,0 +1,194 @@
+// Copyright Epic Games, Inc. All Rights Reserved.
+
+#if WITH_DEV_AUTOMATION_TESTS
+#if WITH_EDITOR
+
+#include "UAFTestsUtilities.h"
+#include "CQTest.h"
+#include "Editor.h"
+#include "ScopedTransaction.h"
+#include "Subsystems/AssetEditorSubsystem.h"
+#include "Graph/AnimNextAnimationGraph.h"
+#include "Graph/AnimNextAnimationGraphFactory.h"
+#include "AnimNextRigVMAssetEditorData.h"
+#include "Math/Vector2D.h"
+
+TEST_CLASS(FAnimNextAnimationGraphTests, "Animation.UAF.Functional")
+{
+	UAnimNextAnimationGraph* AnimNextAnimationGraph = nullptr;
+	UAssetEditorSubsystem* AssetEditorSubsystem = nullptr;
+	UAnimNextRigVMAssetEditorData* EditorData = nullptr;
+	URigVMGraph* ControllerGraph = nullptr;
+	UEdGraph* ParentGraph = nullptr;
+	const FString TraitStackNodeName = "AnimNextTraitStack";
+	const FString BlendByBoolTraitName = "FBlendByBoolTrait";
+	const FString AccumulateTransformLerpNodeName = "RigVMFunction_AccumulateTransformLerp";
+	const FString VariableName = "NewVariable";
+	const FString GetVariableNodeName = "VariableNode";
+	const FString SetVariableNodeName = "VariableNode_1";
+
+	BEFORE_EACH()
+	{
+		AssetEditorSubsystem = GEditor->GetEditorSubsystem<UAssetEditorSubsystem>();
+		ASSERT_THAT(IsNotNull(AssetEditorSubsystem));
+	}
+
+	// QMetry Test Case:  UE-TC-10707
+	TEST_METHOD(Create_Trait_Stack_Node_In_AnimNext_AnimationGraph)
+	{
+		TestCommandBuilder
+			.Do(TEXT("Create AnimNext Animation Graph Asset"), [&]()
+			{
+				CreateAnimNextAnimationGraph();
+			})
+			.Then(TEXT("Create a Trait Stack node"), [&]()
+			{
+				CreateTraitStackNode();
+			});
+	}
+	
+	// QMetry Test Case:  UE-TC-10708
+	TEST_METHOD(Add_Blend_By_Bool_To_Trait_Stack_Node_In_AnimNext_AnimationGraph)
+	{
+		TestCommandBuilder
+			.Do(TEXT("Create AnimNext Animation Graph Asset"), [&]()
+			{
+				CreateAnimNextAnimationGraph();
+			})
+			.Then(TEXT("Create a Trait Stack node"), [&]()
+			{
+				CreateTraitStackNode();
+			})
+			.Then(TEXT("Add Blend by Bool to the Trait Stack node"), [&]()
+			{
+				URigVMController* Controller = EditorData->GetControllerByName("RigVMGraph");
+				FName TraitName = Controller->AddTrait(FName(TraitStackNodeName), FName("/Script/AnimNextAnimGraph.RigDecorator_AnimNextCppDecorator"), FName(BlendByBoolTraitName), "(DecoratorSharedDataStruct=\"/Script/CoreUObject.ScriptStruct'/Script/AnimNextAnimGraph.AnimNextBlendByBoolTraitSharedData'\",Name=\"FBlendByBoolTrait\",TrueChild=(PackedTraitIndexAndNodeHandle=16777215),FalseChild=(PackedTraitIndexAndNodeHandle=16777215),bCondition=False,bAlwaysUpdateTrueChild=False)");
+				ASSERT_THAT(IsTrue(TraitName.ToString().Equals(BlendByBoolTraitName)));
+			});
+	}
+	
+	// QMetry Test Case:  UE-TC-18823
+	TEST_METHOD(Undo_Redo_In_AnimNext_AnimationGraph)
+	{
+		TestCommandBuilder
+			.Do(TEXT("Create AnimNext Animation Graph Asset"), [&]()
+			{
+				CreateAnimNextAnimationGraph();
+			})
+			.Then(TEXT("Create a Trait Stack node"), [&]()
+			{
+				CreateTraitStackNode();
+			})
+			.Then(TEXT("Create an AccumulateLerp Transform node"), [&]()
+			{
+				TArray<UEdGraphPin*> FromPins;
+				UEdGraphNode* AccumulateLerpTransformNode = UAFTestsUtilities::AddUnitNode(ParentGraph, "/Script/RigVM.RigVMFunction_AccumulateTransformLerp", FromPins, FVector2f(-525.0f, 125.0f));
+				ASSERT_THAT(IsNotNull(AccumulateLerpTransformNode));
+			})
+			.Then(TEXT("Create a bool Variable"), [&]()
+			{
+				FScopedTransaction Transaction(NSLOCTEXT("Add_Variables", "AddVariables", "Add Variable(s)"));
+				
+				UAnimNextVariableEntry* VariableEntry = UAFTestsUtilities::AddVariable(AnimNextAnimationGraph, FAnimNextParamType::GetType<bool>(), VariableName, "false");
+				ASSERT_THAT(IsNotNull(VariableEntry));
+			})
+			.Then(TEXT("Add Get Variable Node to Graph"), [&]()
+			{
+				TArray<UEdGraphPin*> FromPins;
+				UEdGraphNode* GetVariableNode = UAFTestsUtilities::AddVariableNode(ParentGraph, AnimNextAnimationGraph, VariableName, FAnimNextParamType::GetType<bool>(), FAnimNextSchemaAction_Variable::EVariableAccessorChoice::Get, FromPins, FVector2f(-250.0f, 135.0f));
+				ASSERT_THAT(IsNotNull(GetVariableNode));
+			})
+			.Then(TEXT("Add Set Variable Node to Graph"), [&]()
+			{				
+				TArray<UEdGraphPin*> FromPins;
+				UEdGraphNode* SetVariableNode = UAFTestsUtilities::AddVariableNode(ParentGraph, AnimNextAnimationGraph, VariableName, FAnimNextParamType::GetType<bool>(), FAnimNextSchemaAction_Variable::EVariableAccessorChoice::Set, FromPins, FVector2f(-250.0f, 235.0f));
+				ASSERT_THAT(IsNotNull(SetVariableNode));
+			})
+			.Until(TEXT("Execute Undo Set Variable Node"), [&]()
+			{
+				return GEditor->UndoTransaction();
+			})
+			.Until(TEXT("Execute Undo Get Variable Node"), [&]()
+			{
+				return GEditor->UndoTransaction();
+			})
+			.Until(TEXT("Execute Undo Create bool Variable"), [&]()
+			{
+				return GEditor->UndoTransaction();
+			})
+			.Until(TEXT("Execute Undo Create AccumulateLerp Transform Node"), [&]()
+			{
+				return GEditor->UndoTransaction();
+			})
+			.Until(TEXT("Execute Undo Create Trait Stack Node"), [&]()
+			{
+				return GEditor->UndoTransaction();
+			})
+			.Then(TEXT("Assert Undos"), [&]()
+			{
+				ASSERT_THAT(IsNull(ControllerGraph->FindNodeByName(FName(SetVariableNodeName))));
+				ASSERT_THAT(IsNull(ControllerGraph->FindNodeByName(FName(GetVariableNodeName))));
+				ASSERT_THAT(IsNull(ControllerGraph->FindNodeByName(FName(AccumulateTransformLerpNodeName))));
+				ASSERT_THAT(IsNull(ControllerGraph->FindNodeByName(FName(TraitStackNodeName))));
+				ASSERT_THAT(IsNull(EditorData->FindEntry(FName(VariableName))));
+			})
+			.Until(TEXT("Execute Redo Create Trait Stack Node"), [&]()
+			{
+				return GEditor->RedoTransaction();
+			})
+			.Until(TEXT("Execute Redo Create AccumulateLerp Transform Node"), [&]()
+			{
+				return GEditor->RedoTransaction();
+			})
+			.Until(TEXT("Execute Redo Create bool Variable"), [&]()
+			{
+				return GEditor->RedoTransaction();
+			})
+			.Until(TEXT("Execute Redo Create Get Variable Node"), [&]()
+			{
+				return GEditor->RedoTransaction();
+			})
+			.Until(TEXT("Execute Redo Create Set Variable Node"), [&]()
+			{
+				return GEditor->RedoTransaction();
+			})
+			.Then(TEXT("Assert Redos"), [&]()
+			{
+				ASSERT_THAT(IsNotNull(ControllerGraph->FindNodeByName(FName(SetVariableNodeName))));
+				ASSERT_THAT(IsNotNull(ControllerGraph->FindNodeByName(FName(GetVariableNodeName))));
+				ASSERT_THAT(IsNotNull(ControllerGraph->FindNodeByName(FName(AccumulateTransformLerpNodeName))));
+				ASSERT_THAT(IsNotNull(ControllerGraph->FindNodeByName(FName(TraitStackNodeName))));
+				ASSERT_THAT(IsNotNull(EditorData->FindEntry(FName(VariableName))));
+			});
+	}
+	
+protected:
+	void CreateAnimNextAnimationGraph()
+	{
+		const TSubclassOf<UAnimNextAnimationGraphFactory> AnimGraphFactoryClass = UAnimNextAnimationGraphFactory::StaticClass();
+		UObject* FactoryObject = UAFTestsUtilities::CreateFactoryObject(NewObject<UAnimNextAnimationGraphFactory>(GetTransientPackage(), AnimGraphFactoryClass.Get()), UAnimNextAnimationGraph::StaticClass(), "NewAnimNextAnimationGraph");
+		AnimNextAnimationGraph = CastChecked<UAnimNextAnimationGraph>(FactoryObject);
+		ASSERT_THAT(IsNotNull(AnimNextAnimationGraph));
+		ASSERT_THAT(IsTrue(AssetEditorSubsystem->OpenEditorForAsset(AnimNextAnimationGraph)));
+	}
+	
+	void CreateTraitStackNode()
+	{
+		EditorData = UE::UAF::UncookedOnly::FUtils::GetEditorData(AnimNextAnimationGraph);
+		ASSERT_THAT(IsNotNull(EditorData));
+		ControllerGraph = EditorData->GetControllerByName("RigVMGraph")->GetGraph();
+		ASSERT_THAT(IsNotNull(ControllerGraph));				
+		ParentGraph = Cast<UEdGraph>(EditorData->GetEditorObjectForRigVMGraph(ControllerGraph));
+		ASSERT_THAT(IsNotNull(ParentGraph));
+		
+		TestRunner->AddExpectedError("Base Trait Data is Invalid. Please, select a new Base Trait.", EAutomationExpectedErrorFlags::Contains, -1);
+		
+		TArray<UEdGraphPin*> FromPins;
+		UEdGraphNode* TraitStackNode = UAFTestsUtilities::AddUnitNode(ParentGraph, "/Script/AnimNextAnimGraph.RigUnit_AnimNextTraitStack", FromPins, FVector2f(0.0f, 0.0f));
+		ASSERT_THAT(IsNotNull(TraitStackNode));
+		ASSERT_THAT(IsNotNull(ControllerGraph->FindNodeByName(FName(TraitStackNodeName))));	
+	}
+};
+
+#endif // WITH_EDITOR
+#endif // WITH_DEV_AUTOMATION_TESTS
diff --git a/Engine/Plugins/Experimental/UAF/UAFTests/Source/UAFTests/Private/UAFAssetCreationTests.cpp b/Engine/Plugins/Experimental/UAF/UAFTests/Source/UAFTests/Private/UAFAssetCreationTests.cpp
new file mode 100644
--- /dev/null
+++ b/Engine/Plugins/Experimental/UAF/UAFTests/Source/UAFTests/Private/UAFAssetCreationTests.cpp
@@ -0,0 +1,115 @@
+// Copyright Epic Games, Inc. All Rights Reserved.
+
+#if WITH_DEV_AUTOMATION_TESTS
+#if WITH_EDITOR
+
+#include "UAFTestsUtilities.h"
+#include "CQTest.h"
+#include "Editor.h"
+#include "Subsystems/AssetEditorSubsystem.h"
+#include "WorkspaceEditor/Private/Workspace.h"
+#include "Workspace/AnimNextWorkspaceFactory.h"
+#include "AnimNextStateTree.h"
+#include "AnimNextStateTreeFactory.h"
+#include "Graph/AnimNextAnimationGraph.h"
+#include "Graph/AnimNextAnimationGraphFactory.h"
+#include "Module/AnimNextModule.h"
+#include "Module/AnimNextModuleFactory.h"
+#include "Variables/AnimNextSharedVariables.h"
+#include "Variables/AnimNextSharedVariablesFactory.h"
+
+// QMetry Test Case:  UE-TC-10715 
+TEST_CLASS(FAssetCreationTests, "Animation.UAF.Functional")
+{
+	UObject* AnimNextAsset;
+	UAssetEditorSubsystem* AssetEditorSubsystem;
+
+	BEFORE_EACH()
+	{
+		AssetEditorSubsystem = GEditor->GetEditorSubsystem<UAssetEditorSubsystem>();
+		ASSERT_THAT(IsNotNull(AssetEditorSubsystem));
+	}
+
+	TEST_METHOD(Verify_Workspace_Asset_Creation)
+	{		
+		TestCommandBuilder
+			.Do(TEXT("Create AnimNext Workspace Asset"), [&]()
+			{
+				const TSubclassOf<UAnimNextWorkspaceFactory> WorkSpaceFactoryClass = UAnimNextWorkspaceFactory::StaticClass();
+				UObject* FactoryObject = UAFTestsUtilities::CreateFactoryObject(NewObject<UAnimNextWorkspaceFactory>(), UWorkspace::StaticClass(), "NewAnimNextWorkspace");
+				AnimNextAsset = CastChecked<UWorkspace>(FactoryObject);
+			})
+			.Then(TEXT("Assert AnimNext Workspace Asset Created and Open in Editor"), [&]()
+			{
+				ASSERT_THAT(IsNotNull(AnimNextAsset));
+				ASSERT_THAT(IsTrue(AssetEditorSubsystem->OpenEditorForAsset(AnimNextAsset)));
+			});
+	}
+	
+	TEST_METHOD(Verify_StateTree_Asset_Creation)
+	{
+		TestCommandBuilder
+			.Do(TEXT("Create AnimNext State Tree Asset"), [&]()
+			{
+				const TSubclassOf<UAnimNextStateTreeFactory> StateTreeFactoryClass = UAnimNextStateTreeFactory::StaticClass();
+				UObject* FactoryObject = UAFTestsUtilities::CreateFactoryObject(NewObject<UAnimNextStateTreeFactory>(), UAnimNextStateTree::StaticClass(), "NewAnimNextStateTree");
+				AnimNextAsset = CastChecked<UAnimNextStateTree>(FactoryObject);
+			})
+			.Then(TEXT("Assert AnimNext State Tree Asset Created and Open in Editor"), [&]()
+			{
+				ASSERT_THAT(IsNotNull(AnimNextAsset));
+				ASSERT_THAT(IsTrue(AssetEditorSubsystem->OpenEditorForAsset(AnimNextAsset)));
+			});
+	}
+	
+	TEST_METHOD(Verify_AnimGraph_Asset_Creation)
+	{
+		TestCommandBuilder
+			.Do(TEXT("Create AnimNext Animation Graph Asset"), [&]()
+			{
+				const TSubclassOf<UAnimNextAnimationGraphFactory> AnimGraphFactoryClass = UAnimNextAnimationGraphFactory::StaticClass();
+				UObject* FactoryObject = UAFTestsUtilities::CreateFactoryObject(NewObject<UAnimNextAnimationGraphFactory>(), UAnimNextAnimationGraph::StaticClass(), "NewAnimNextAnimationGraph");
+				AnimNextAsset = CastChecked<UAnimNextAnimationGraph>(FactoryObject);
+			})
+			.Then(TEXT("Assert AnimNext Animation Graph Asset Created and Open in Editor"), [&]()
+			{
+				ASSERT_THAT(IsNotNull(AnimNextAsset));
+				ASSERT_THAT(IsTrue(AssetEditorSubsystem->OpenEditorForAsset(AnimNextAsset)));
+			});		
+	}
+	
+	TEST_METHOD(Verify_Module_Asset_Creation)
+	{
+		TestCommandBuilder
+			.Do(TEXT("Create AnimNext Module Asset"), [&]()
+			{
+				const TSubclassOf<UAnimNextModuleFactory> ModuleFactoryClass = UAnimNextModuleFactory::StaticClass();
+				UObject* FactoryObject = UAFTestsUtilities::CreateFactoryObject(NewObject<UAnimNextModuleFactory>(), UAnimNextModule::StaticClass(), "NewAnimNextModule");				
+				AnimNextAsset = CastChecked<UAnimNextModule>(FactoryObject);
+			})
+			.Then(TEXT("Assert AnimNext Module Asset Created and Open in Editor"), [&]()
+			{
+				ASSERT_THAT(IsNotNull(AnimNextAsset));
+				ASSERT_THAT(IsTrue(AssetEditorSubsystem->OpenEditorForAsset(AnimNextAsset)));
+			});
+	}
+	
+	TEST_METHOD(Verify_SharedVariables_Asset_Creation)
+	{
+		TestCommandBuilder
+			.Do(TEXT("Create AnimNext Shared Variables Asset"), [&]()
+			{
+				const TSubclassOf<UAnimNextSharedVariablesFactory> SharedVariablesFactoryClass = UAnimNextSharedVariablesFactory::StaticClass();
+				UObject* FactoryObject = UAFTestsUtilities::CreateFactoryObject(NewObject<UAnimNextSharedVariablesFactory>(), UAnimNextSharedVariables::StaticClass(), "NewAnimNextSharedVariables");
+				AnimNextAsset = CastChecked<UAnimNextSharedVariables>(FactoryObject);
+			})
+			.Then(TEXT("Assert AnimNext Shared Variables Asset Created and Open in Editor"), [&]()
+			{
+				ASSERT_THAT(IsNotNull(AnimNextAsset));
+				ASSERT_THAT(IsTrue(AssetEditorSubsystem->OpenEditorForAsset(AnimNextAsset)));		
+			});
+	}	
+};
+
+#endif // WITH_EDITOR
+#endif // WITH_DEV_AUTOMATION_TESTS
diff --git a/Engine/Plugins/Experimental/UAF/UAFTests/Source/UAFTests/Private/UAFBlueprintTests.cpp b/Engine/Plugins/Experimental/UAF/UAFTests/Source/UAFTests/Private/UAFBlueprintTests.cpp
new file mode 100644
--- /dev/null
+++ b/Engine/Plugins/Experimental/UAF/UAFTests/Source/UAFTests/Private/UAFBlueprintTests.cpp
@@ -0,0 +1,105 @@
+// Copyright Epic Games, Inc. All Rights Reserved.
+
+#if WITH_DEV_AUTOMATION_TESTS
+#if WITH_EDITOR
+
+#include "UAFTestsUtilities.h"
+#include "CQTest.h"
+#include "Editor.h"
+#include "Component/AnimNextComponent.h"
+#include "Components/ActorComponent.h"
+#include "Factories/BlueprintFactory.h"
+#include "GameFramework/Character.h"
+#include "Kismet2/KismetEditorUtilities.h"
+#include "Subsystems/AssetEditorSubsystem.h"
+#include "Module/AnimNextModule.h"
+#include "Module/AnimNextModuleFactory.h"
+
+TEST_CLASS(FAssetBlueprintTests, "Animation.UAF.Functional")
+{
+	UBlueprint* Blueprint;
+	UObject* UAFAsset;
+	UAnimNextComponent* AnimNextComponent;
+	UAssetEditorSubsystem* AssetEditorSubsystem;
+
+	BEFORE_EACH()
+	{
+		AssetEditorSubsystem = GEditor->GetEditorSubsystem<UAssetEditorSubsystem>();
+		ASSERT_THAT(IsNotNull(AssetEditorSubsystem));
+	}
+
+	// QMetry Test Case:  UE-TC-10710 
+	TEST_METHOD(Verify_Blueprint_Add_AnimNext_Component)
+	{
+		TestCommandBuilder
+			.Do(TEXT("Create Character Blueprint Asset"), [&]()
+			{
+				UBlueprintFactory* Factory = NewObject<UBlueprintFactory>();
+				Factory->ParentClass = ACharacter::StaticClass();
+				
+				const FString BlueprintName = "NewBlueprint";
+				UObject* factoryObject = UAFTestsUtilities::CreateFactoryObject(Factory, UBlueprint::StaticClass(), BlueprintName);
+				Blueprint = CastChecked<UBlueprint>(factoryObject);
+				ASSERT_THAT(IsNotNull(Blueprint));
+				ASSERT_THAT(AreEqual(BlueprintName, Blueprint->GetFName()));
+			})
+			.Then(TEXT("Create AnimNext Component"), [&]()
+			{
+				const FName ComponentName = FName(TEXT("AnimNext"));
+				AnimNextComponent = NewObject<UAnimNextComponent>(GetTransientPackage(), UAnimNextComponent::StaticClass(), ComponentName);
+				ASSERT_THAT(IsNotNull(AnimNextComponent));
+				ASSERT_THAT(AreEqual(ComponentName, AnimNextComponent->GetFName()));
+			})
+			.Then(TEXT("Add the AnimNext Component to the Character Blueprint and Open in Editor"), [&]()
+			{
+				TArray<UActorComponent*> Components;
+				Components.Add(AnimNextComponent);
+				FKismetEditorUtilities::AddComponentsToBlueprint(Blueprint, Components);
+				ASSERT_THAT(IsTrue(AssetEditorSubsystem->OpenEditorForAsset(Blueprint)));
+			});
+	}
+	
+	// QMetry Test Case:  UE-TC-10711
+	TEST_METHOD(Verify_Assign_Module_To_Blueprint)
+	{
+		TestCommandBuilder
+			.Do(TEXT("Create UAF Module Asset"), [&]()
+			{
+				const FString UAFModuleName = "NewAnimNextModule";
+				const TSubclassOf<UAnimNextModuleFactory> ModuleFactoryClass = UAnimNextModuleFactory::StaticClass();
+				UObject* factoryObject = UAFTestsUtilities::CreateFactoryObject(NewObject<UAnimNextModuleFactory>(), UAnimNextModule::StaticClass(), UAFModuleName);				
+				UAFAsset = CastChecked<UAnimNextModule>(factoryObject);			
+				ASSERT_THAT(IsNotNull(UAFAsset));
+				ASSERT_THAT(AreEqual(UAFModuleName, UAFAsset->GetFName()));
+			})			
+			.Then(TEXT("Create Character Blueprint Asset"), [&]()
+			{
+				UBlueprintFactory* Factory = NewObject<UBlueprintFactory>();
+				Factory->ParentClass = ACharacter::StaticClass();
+				
+				const FString BlueprintName = "NewBlueprint";
+				UObject* factoryObject = UAFTestsUtilities::CreateFactoryObject(Factory, UBlueprint::StaticClass(), BlueprintName);
+				Blueprint = CastChecked<UBlueprint>(factoryObject);
+				ASSERT_THAT(IsNotNull(Blueprint));
+				ASSERT_THAT(AreEqual(BlueprintName, Blueprint->GetFName()));
+			})	
+			.Then(TEXT("Create AnimNext Component"), [&]()
+			{
+				const FName ComponentName = FName(TEXT("AnimNext"));
+				AnimNextComponent = NewObject<UAnimNextComponent>(UAFAsset->GetPackage(), UAnimNextComponent::StaticClass(), ComponentName);
+				ASSERT_THAT(IsNotNull(AnimNextComponent));
+				ASSERT_THAT(AreEqual(ComponentName, AnimNextComponent->GetFName()));
+			})
+			.Then(TEXT("Set the UAF Module and add the AnimNext Component to the Character Blueprint and Open in Editor"), [&]()
+			{
+				TArray<UActorComponent*> Components;
+				AnimNextComponent->SetModule(Cast<UAnimNextModule>(UAFAsset));
+				Components.Add(AnimNextComponent);
+				FKismetEditorUtilities::AddComponentsToBlueprint(Blueprint, Components);
+				ASSERT_THAT(IsTrue(AssetEditorSubsystem->OpenEditorForAsset(Blueprint)));
+			});
+	}
+};
+
+#endif // WITH_EDITOR
+#endif // WITH_DEV_AUTOMATION_TESTS
diff --git a/Engine/Plugins/Experimental/UAF/UAFTests/Source/UAFTests/Private/UAFFunctionNodeTests.cpp b/Engine/Plugins/Experimental/UAF/UAFTests/Source/UAFTests/Private/UAFFunctionNodeTests.cpp
new file mode 100644
--- /dev/null
+++ b/Engine/Plugins/Experimental/UAF/UAFTests/Source/UAFTests/Private/UAFFunctionNodeTests.cpp
@@ -0,0 +1,159 @@
+// Copyright Epic Games, Inc. All Rights Reserved.
+
+#if WITH_DEV_AUTOMATION_TESTS
+#if WITH_EDITOR
+
+#include "UAFTestsUtilities.h"
+#include "CQTest.h"
+#include "Editor.h"
+#include "Subsystems/AssetEditorSubsystem.h"
+#include "Graph/AnimNextAnimationGraph.h"
+#include "Graph/AnimNextAnimationGraphFactory.h"
+#include "Module/AnimNextModule.h"
+#include "Module/AnimNextModuleFactory.h"
+#include "AnimNextStateTree.h"
+#include "AnimNextStateTreeFactory.h"
+#include "AnimNextRigVMAssetEditorData.h"
+#include "RigVMModel/Nodes/RigVMLibraryNode.h"
+
+TEST_CLASS(FAssetFunctionNodeTests, "Animation.UAF.Functional")
+{
+	UAnimNextRigVMAsset* AnimNextAsset;
+	UAssetEditorSubsystem* AssetEditorSubsystem;
+	UAnimNextRigVMAssetEditorData* EditorData;
+	URigVMLibraryNode* LibraryNode;
+	const FString FunctionName = "NewFunction";
+
+	BEFORE_EACH()
+	{
+		AssetEditorSubsystem = GEditor->GetEditorSubsystem<UAssetEditorSubsystem>();
+		ASSERT_THAT(IsNotNull(AssetEditorSubsystem));
+	}
+
+	// QMetry Test Case:  UE-TC-18834
+	TEST_METHOD(AnimNextAnimationGraph_Create_Function_Node_Pins)
+	{		
+		TestCommandBuilder
+			.Do(TEXT("Create AnimNext Animation Graph Asset"), [&]()
+			{
+				const TSubclassOf<UAnimNextAnimationGraphFactory> AnimNextAnimGraphFactoryClass = UAnimNextAnimationGraphFactory::StaticClass();
+				UObject* factoryObject = UAFTestsUtilities::CreateFactoryObject(NewObject<UAnimNextAnimationGraphFactory>(GetTransientPackage(), AnimNextAnimGraphFactoryClass.Get()), UAnimNextAnimationGraph::StaticClass(), "NewAnimNextAnimationGraph");
+				AnimNextAsset = CastChecked<UAnimNextAnimationGraph>(factoryObject);				
+				ASSERT_THAT(IsTrue(AssetEditorSubsystem->OpenEditorForAsset(AnimNextAsset)));
+			})
+			.Then(TEXT("Create Function Node"), [&]()
+			{
+				LibraryNode = UAFTestsUtilities::AddFunctionNode(AnimNextAsset);
+				ASSERT_THAT(IsNotNull(LibraryNode));
+				ASSERT_THAT(AreEqual(FunctionName, LibraryNode->GetNodeTitle()));
+			})
+			.Then(TEXT("Assert Function Node Creation"), [&]()
+			{
+				EditorData = UE::UAF::UncookedOnly::FUtils::GetEditorData(AnimNextAsset);	
+				URigVMController* Controller = CastChecked<IRigVMClientHost>(EditorData)->GetControllerByName(FunctionName);
+				ASSERT_THAT(IsNotNull(Controller));
+				ASSERT_THAT(IsNotNull(Controller->GetGraph()));				
+			})		
+			.Then(TEXT("Add Function Node Pins"), [&]()
+			{
+				const FString PinName = "Argument";
+				URigVMPin* InputPin = UAFTestsUtilities::AddPin(AnimNextAsset, LibraryNode, ERigVMPinDirection::Input, PinName);
+				URigVMPin* OutputPin =  UAFTestsUtilities::AddPin(AnimNextAsset, LibraryNode, ERigVMPinDirection::Output, PinName);
+				ASSERT_THAT(IsNotNull(InputPin));
+				ASSERT_THAT(IsNotNull(OutputPin));
+			});
+	}
+	
+	// QMetry Test Case:  UE-TC-18831
+	TEST_METHOD(Undo_Redo_AnimNextModule_Create_Function_Node)
+	{		
+		TestCommandBuilder
+			.Do(TEXT("Create AnimNext Module Asset"), [&]()
+			{
+				const TSubclassOf<UAnimNextModuleFactory> AnimNextModuleFactoryClass = UAnimNextModuleFactory::StaticClass();
+				UObject* factoryObject = UAFTestsUtilities::CreateFactoryObject(NewObject<UAnimNextModuleFactory>(GetTransientPackage(), AnimNextModuleFactoryClass.Get()), UAnimNextModule::StaticClass(), "NewAnimNextModule");
+				AnimNextAsset = CastChecked<UAnimNextModule>(factoryObject);				
+				ASSERT_THAT(IsTrue(AssetEditorSubsystem->OpenEditorForAsset(AnimNextAsset)));
+			})
+			.Then(TEXT("Create Function Node"), [&]()
+			{
+				LibraryNode = UAFTestsUtilities::AddFunctionNode(AnimNextAsset);
+				ASSERT_THAT(IsNotNull(LibraryNode));
+				ASSERT_THAT(AreEqual(FunctionName, LibraryNode->GetNodeTitle()));
+			})
+			.Then(TEXT("Assert Function Node Creation"), [&]()
+			{
+				EditorData = UE::UAF::UncookedOnly::FUtils::GetEditorData(AnimNextAsset);	
+				URigVMController* Controller = CastChecked<IRigVMClientHost>(EditorData)->GetControllerByName(FunctionName);
+				ASSERT_THAT(IsNotNull(Controller));
+				ASSERT_THAT(IsNotNull(Controller->GetGraph()));				
+			})
+			.Until(TEXT("Execute Undo"), [&]()
+			{
+				return GEditor->UndoTransaction();
+			})	
+			.Then(TEXT("Assert Undo Function Node"), [&]()
+			{
+				URigVMController* Controller = CastChecked<IRigVMClientHost>(EditorData)->GetControllerByName(FunctionName);	
+				ASSERT_THAT(IsNull(Controller));
+			})
+			.Until(TEXT("Execute Redo"), [&]()
+			{
+				return GEditor->RedoTransaction();
+			})
+			.Then(TEXT("Assert Redo Function Node"), [&]()
+			{
+				URigVMController* Controller = CastChecked<IRigVMClientHost>(EditorData)->GetControllerByName(FunctionName);	
+				ASSERT_THAT(IsNotNull(Controller));
+				ASSERT_THAT(IsNotNull(Controller->GetGraph()));
+			});
+	}
+	
+	// QMetry Test Case:  UE-TC-18829
+	TEST_METHOD(Undo_Redo_AnimNextAnimationGraph_Create_Function_Node)
+	{		
+		TestCommandBuilder
+			.Do(TEXT("Create AnimNext Animation Graph Asset"), [&]()
+			{
+				const TSubclassOf<UAnimNextAnimationGraphFactory> AnimNextAnimGraphFactoryClass = UAnimNextAnimationGraphFactory::StaticClass();
+				UObject* factoryObject = UAFTestsUtilities::CreateFactoryObject(NewObject<UAnimNextAnimationGraphFactory>(GetTransientPackage(), AnimNextAnimGraphFactoryClass.Get()), UAnimNextAnimationGraph::StaticClass(), "NewAnimNextAnimationGraph");
+				AnimNextAsset = CastChecked<UAnimNextAnimationGraph>(factoryObject);				
+				ASSERT_THAT(IsTrue(AssetEditorSubsystem->OpenEditorForAsset(AnimNextAsset)));
+			})
+			.Then(TEXT("Create Function Node"), [&]()
+			{
+				LibraryNode = UAFTestsUtilities::AddFunctionNode(AnimNextAsset);
+				ASSERT_THAT(IsNotNull(LibraryNode));
+				ASSERT_THAT(AreEqual(FunctionName, LibraryNode->GetNodeTitle()));
+			})
+			.Then(TEXT("Assert Function Node Creation"), [&]()
+			{
+				EditorData = UE::UAF::UncookedOnly::FUtils::GetEditorData(AnimNextAsset);	
+				URigVMController* Controller = CastChecked<IRigVMClientHost>(EditorData)->GetControllerByName(FunctionName);
+				ASSERT_THAT(IsNotNull(Controller));
+				ASSERT_THAT(IsNotNull(Controller->GetGraph()));				
+			})
+			.Until(TEXT("Execute Undo"), [&]()
+			{
+				return GEditor->UndoTransaction();
+			})	
+			.Then(TEXT("Assert Undo Function Node"), [&]()
+			{
+				URigVMController* Controller = CastChecked<IRigVMClientHost>(EditorData)->GetControllerByName(FunctionName);	
+				ASSERT_THAT(IsNull(Controller));
+			})
+			.Until(TEXT("Execute Redo"), [&]()
+			{
+				return GEditor->RedoTransaction();
+			})
+			.Then(TEXT("Assert Redo Function Node"), [&]()
+			{
+				URigVMController* Controller = CastChecked<IRigVMClientHost>(EditorData)->GetControllerByName(FunctionName);	
+				ASSERT_THAT(IsNotNull(Controller));
+				ASSERT_THAT(IsNotNull(Controller->GetGraph()));
+			});
+	}	
+};
+
+#endif // WITH_EDITOR
+#endif // WITH_DEV_AUTOMATION_TESTS
diff --git a/Engine/Plugins/Experimental/UAF/UAFTests/Source/UAFTests/Private/UAFModuleTests.cpp b/Engine/Plugins/Experimental/UAF/UAFTests/Source/UAFTests/Private/UAFModuleTests.cpp
new file mode 100644
--- /dev/null
+++ b/Engine/Plugins/Experimental/UAF/UAFTests/Source/UAFTests/Private/UAFModuleTests.cpp
@@ -0,0 +1,345 @@
+// Copyright Epic Games, Inc. All Rights Reserved.
+
+#if WITH_DEV_AUTOMATION_TESTS
+#if WITH_EDITOR
+
+#include "UAFTestsUtilities.h"
+#include "CQTest.h"
+#include "Editor.h"
+#include "ScopedTransaction.h"
+#include "Subsystems/AssetEditorSubsystem.h"
+#include "WorkspaceEditor/Private/Workspace.h"
+#include "Workspace/AnimNextWorkspaceFactory.h"
+#include "Module/AnimNextModule.h"
+#include "Module/AnimNextModuleFactory.h"
+#include "AnimNextRigVMAssetEditorData.h"
+#include "Math/Vector2D.h"
+#include "AssetToolsModule.h"
+#include "AssetRegistry/AssetRegistryModule.h"
+#include "ObjectTools.h"
+#include "UObject/SavePackage.h"
+
+TEST_CLASS(FAnimNextModuleTests, "Animation.UAF.Functional")
+{
+	UWorkspace* AnimNextWorkspace = nullptr;
+	UAnimNextModule* AnimNextModule = nullptr;
+	UPackage* AnimNextWorkspacePackage = nullptr;
+	UPackage* AnimNextModulePackage = nullptr;
+	FName AnimNextWorkspacePackageName;
+	FName AnimNextModulePackageName;
+	UAssetEditorSubsystem* AssetEditorSubsystem = nullptr;
+	UAnimNextRigVMAssetEditorData* EditorData = nullptr;
+	URigVMGraph* ControllerGraph = nullptr;
+	UEdGraph* ParentGraph = nullptr;
+	UEdGraphNode* RunGraphNode = nullptr;
+	UEdGraphNode* WritePoseNode = nullptr;
+	const FString InitializeNodeName = "AnimNextInitializeEvent";
+	const FString RunGraphNodeName = "AnimNextRunAnimationGraph_v2";
+	const FString WritePoseNodeName = "AnimNextWriteSkeletalMeshComponentPose";
+	const FString CollapseNodeName = "CollapseNode";
+	const FString FunctionNodeName = "New Function";
+	const FString VariableName = "NewVariable";
+	const FString GetVariableNodeName = "VariableNode";
+	const FString SetVariableNodeName = "VariableNode_1";
+
+	BEFORE_EACH()
+	{
+		AssetEditorSubsystem = GEditor->GetEditorSubsystem<UAssetEditorSubsystem>();
+		ASSERT_THAT(IsNotNull(AssetEditorSubsystem));
+	}
+	
+	AFTER_EACH()
+	{
+		if (AnimNextWorkspace)
+		{
+			ObjectTools::DeleteSingleObject(AnimNextWorkspace);
+			AnimNextWorkspace = nullptr;
+		}
+		
+		if (AnimNextModule)
+		{
+			ObjectTools::DeleteSingleObject(AnimNextModule);
+			AnimNextModule = nullptr;
+		}
+	}
+	
+	// QMetry Test Case:  UE-TC-10714
+	TEST_METHOD(Save_AnimNext_Module_Without_Workspace)
+	{
+		TestCommandBuilder
+			.Do(TEXT("Create AnimNext Workspaxce Asset"), [&]()
+			{
+				CreateAnimNextWorkspace();
+			})
+			.Then(TEXT("Create AnimNext Module Asset"), [&]()
+			{
+				CreateAnimNextModule();
+			})
+			.Then(TEXT("Assert that UAF Assets were created"), [&]()
+			{
+				ASSERT_THAT(IsNotNull(AnimNextWorkspace));
+				ASSERT_THAT(IsTrue(AssetEditorSubsystem->OpenEditorForAsset(AnimNextWorkspace)));
+				ASSERT_THAT(IsNotNull(AnimNextModule));
+				ASSERT_THAT(IsTrue(AssetEditorSubsystem->OpenEditorForAsset(AnimNextModule)));
+			})
+			.Then(TEXT("Assert that UAF Assets were saved"), [&]()
+			{
+				FSavePackageArgs SaveArgs;
+				SaveArgs.TopLevelFlags = RF_Standalone;				
+				ASSERT_THAT(IsTrue(UPackage::SavePackage(AnimNextWorkspacePackage, nullptr, *(AnimNextWorkspacePackageName.ToString()), SaveArgs)));
+				ASSERT_THAT(IsTrue(UPackage::SavePackage(AnimNextModulePackage, nullptr, *(AnimNextModulePackageName.ToString()), SaveArgs)));
+			});
+	}
+		
+	// QMetry Test Case:  UE-TC-18778
+	TEST_METHOD(Undo_Redo_In_AnimNext_Module)
+	{
+		TestCommandBuilder
+			.Do(TEXT("Create AnimNext Module Asset"), [&]()
+			{
+				CreateAnimNextModule();
+			})
+			.Then(TEXT("Create a Write Pose node"), [&]()
+			{
+				EditorData = UE::UAF::UncookedOnly::FUtils::GetEditorData(AnimNextModule);
+				ASSERT_THAT(IsNotNull(EditorData));
+				ControllerGraph = EditorData->GetControllerByName("RigVMGraph")->GetGraph();
+				ASSERT_THAT(IsNotNull(ControllerGraph));				
+				ParentGraph = Cast<UEdGraph>(EditorData->GetEditorObjectForRigVMGraph(ControllerGraph));
+				ASSERT_THAT(IsNotNull(ParentGraph));
+				
+				TArray<UEdGraphPin*> FromPins;
+				WritePoseNode = UAFTestsUtilities::AddUnitNode(ParentGraph, "/Script/AnimNext.RigUnit_AnimNextWriteSkeletalMeshComponentPose", FromPins, FVector2f(250.0f, 135.0f));
+				ASSERT_THAT(IsNotNull(WritePoseNode));
+			})
+			.Then(TEXT("Create a bool Variable"), [&]()
+			{
+				FScopedTransaction Transaction(NSLOCTEXT("Add_Variables", "AddVariables", "Add Variable(s)"));
+
+				UAnimNextVariableEntry* VariableEntry = UAFTestsUtilities::AddVariable(AnimNextModule, FAnimNextParamType::GetType<bool>(), VariableName, "false");
+				ASSERT_THAT(IsNotNull(VariableEntry));
+			})
+			.Then(TEXT("Add Get Variable Node to Graph"), [&]()
+			{
+				TArray<UEdGraphPin*> FromPins;
+				UEdGraphNode* GetVariableNode = UAFTestsUtilities::AddVariableNode(ParentGraph, AnimNextModule, VariableName, FAnimNextParamType::GetType<bool>(), FAnimNextSchemaAction_Variable::EVariableAccessorChoice::Get, FromPins, FVector2f(-250.0f, 135.0f));
+				ASSERT_THAT(IsNotNull(GetVariableNode));
+			})
+			.Then(TEXT("Add Set Variable Node to Graph"), [&]()
+			{
+				TArray<UEdGraphPin*> FromPins;
+				UEdGraphNode* SetVariableNode = UAFTestsUtilities::AddVariableNode(ParentGraph, AnimNextModule, VariableName, FAnimNextParamType::GetType<bool>(), FAnimNextSchemaAction_Variable::EVariableAccessorChoice::Set, FromPins, FVector2f(-250.0f, 235.0f));
+				ASSERT_THAT(IsNotNull(SetVariableNode));
+			})
+			.Until(TEXT("Execute Undo Set Variable Node"), [&]()
+			{
+				return GEditor->UndoTransaction();
+			})
+			.Until(TEXT("Execute Undo Get Variable Node"), [&]()
+			{
+				return GEditor->UndoTransaction();
+			})
+			.Until(TEXT("Execute Undo Create bool Variable"), [&]()
+			{
+				return GEditor->UndoTransaction();
+			})		
+			.Until(TEXT("Execute Undo Create Write Pose Node"), [&]()
+			{
+				return GEditor->UndoTransaction();
+			})
+			.Then(TEXT("Assert Undos"), [&]()
+			{
+				ASSERT_THAT(IsNull(ControllerGraph->FindNodeByName(FName(SetVariableNodeName))));
+				ASSERT_THAT(IsNull(ControllerGraph->FindNodeByName(FName(GetVariableNodeName))));
+				ASSERT_THAT(IsNull(ControllerGraph->FindNodeByName(FName(WritePoseNodeName))));
+				ASSERT_THAT(IsNull(EditorData->FindEntry(FName(VariableName))));
+			})
+			.Until(TEXT("Execute Redo Create Write Pose Node"), [&]()
+			{
+				return GEditor->RedoTransaction();
+			})
+			.Until(TEXT("Execute Redo Create bool Variable"), [&]()
+			{
+				return GEditor->RedoTransaction();
+			})
+			.Until(TEXT("Execute Redo Get Variable Node"), [&]()
+			{
+				return GEditor->RedoTransaction();
+			})
+			.Until(TEXT("Execute Redo Set Variable Node"), [&]()
+			{
+				return GEditor->RedoTransaction();
+			})
+			.Then(TEXT("Assert Redos"), [&]()
+			{
+				ASSERT_THAT(IsNotNull(ControllerGraph->FindNodeByName(FName(SetVariableNodeName))));
+				ASSERT_THAT(IsNotNull(ControllerGraph->FindNodeByName(FName(GetVariableNodeName))));
+				ASSERT_THAT(IsNotNull(ControllerGraph->FindNodeByName(FName(WritePoseNodeName))));
+			});
+	}
+
+	// QMetry Test Case:  UE-TC-18832
+	TEST_METHOD(Collapse_Node_To_Function_In_AnimNext_Module)
+	{
+		TestCommandBuilder
+			.Do(TEXT("Create AnimNext Module Asset"), [&]()
+			{
+				CreateAnimNextModule();
+			})
+			.Then(TEXT("Create a Run Graph node"), [&]()
+			{
+				EditorData = UE::UAF::UncookedOnly::FUtils::GetEditorData(AnimNextModule);
+				ASSERT_THAT(IsNotNull(EditorData));
+				ControllerGraph = EditorData->GetControllerByName("RigVMGraph")->GetGraph();
+				ASSERT_THAT(IsNotNull(ControllerGraph));				
+				ParentGraph = Cast<UEdGraph>(EditorData->GetEditorObjectForRigVMGraph(ControllerGraph));
+				ASSERT_THAT(IsNotNull(ParentGraph));	
+				
+				TArray<UEdGraphPin*> FromPins;
+				RunGraphNode = UAFTestsUtilities::AddUnitNode(ParentGraph, "/Script/AnimNextAnimGraph.RigUnit_AnimNextRunAnimationGraph_v2", FromPins, FVector2f(64.0f, 16.0f));
+				ASSERT_THAT(IsNotNull(RunGraphNode));
+			})
+			.Then(TEXT("Create a Write Pose node"), [&]()
+			{
+				TArray<UEdGraphPin*> FromPins;
+				WritePoseNode = UAFTestsUtilities::AddUnitNode(ParentGraph, "/Script/AnimNext.RigUnit_AnimNextWriteSkeletalMeshComponentPose", FromPins, FVector2f(336.0f, 0.0f));
+				ASSERT_THAT(IsNotNull(WritePoseNode));
+			})
+			.Then(TEXT("Add link from PrePhysics node Execute pin to Run Graph node Execute pin"), [&]()
+			{
+				ASSERT_THAT(IsTrue(UAFTestsUtilities::AddLink(AnimNextModule, "RigUnit_AnimNextPrePhysicsEvent.ExecuteContext", "AnimNextRunAnimationGraph_v2.ExecuteContext")));
+			})
+			.Then(TEXT("Add link from Run Graph node output pin to Write Pose node Execute pin"), [&]()
+			{				
+				ASSERT_THAT(IsTrue(UAFTestsUtilities::AddLink(AnimNextModule, "AnimNextRunAnimationGraph_v2.ExecuteContext", "AnimNextWriteSkeletalMeshComponentPose.ExecuteContext")));
+			})
+			.Then(TEXT("Collpase Run Graph and Write Pose nodes to a Function"), [&]()
+			{
+				TArray<FName> Nodes;
+				Nodes.Add(FName(RunGraphNodeName));
+				Nodes.Add(FName(WritePoseNodeName));
+				URigVMCollapseNode* CollapseNode = UAFTestsUtilities::CollapseNodes(AnimNextModule, Nodes, FunctionNodeName, true);
+				ASSERT_THAT(IsNotNull(CollapseNode));
+			})		
+			.Until(TEXT("Execute Undo Collpase Nodes"), [&]()
+			{
+				return GEditor->UndoTransaction();
+			})
+			.Then(TEXT("Assert Undo Collapse Nodes"), [&]()
+			{
+				ASSERT_THAT(IsNotNull(ControllerGraph->FindNodeByName(FName(RunGraphNodeName))));
+				ASSERT_THAT(IsNotNull(ControllerGraph->FindNodeByName(FName(WritePoseNodeName))));
+				ASSERT_THAT(IsNull(ControllerGraph->FindNodeByName(FName(FunctionNodeName))));
+			})
+			.Until(TEXT("Execute Redo Collapse Nodes"), [&]()
+			{
+				return GEditor->RedoTransaction();
+			})
+			.Then(TEXT("Assert Redo Collapse Nodes"), [&]()
+			{
+				ASSERT_THAT(IsNull(ControllerGraph->FindNodeByName(FName(RunGraphNodeName))));
+				ASSERT_THAT(IsNull(ControllerGraph->FindNodeByName(FName(WritePoseNodeName))));
+				ASSERT_THAT(IsNotNull(ControllerGraph->FindNodeByName(FName(FunctionNodeName))));
+			});		
+	}
+
+	// QMetry Test Case:  UE-TC-18833
+	TEST_METHOD(Collapse_Node_To_Node_In_AnimNext_Module)
+	{
+		TestCommandBuilder
+			.Do(TEXT("Create AnimNext Module Asset"), [&]()
+			{
+				CreateAnimNextModule();
+			})
+			.Then(TEXT("Create a Run Graph node"), [&]()
+			{
+				EditorData = UE::UAF::UncookedOnly::FUtils::GetEditorData(AnimNextModule);
+				ASSERT_THAT(IsNotNull(EditorData));
+				ControllerGraph = EditorData->GetControllerByName("RigVMGraph")->GetGraph();
+				ASSERT_THAT(IsNotNull(ControllerGraph));				
+				ParentGraph = Cast<UEdGraph>(EditorData->GetEditorObjectForRigVMGraph(ControllerGraph));
+				ASSERT_THAT(IsNotNull(ParentGraph));	
+				
+				TArray<UEdGraphPin*> FromPins;
+				RunGraphNode = UAFTestsUtilities::AddUnitNode(ParentGraph, "/Script/AnimNextAnimGraph.RigUnit_AnimNextRunAnimationGraph_v2", FromPins, FVector2f(64.0f, 16.0f));
+				ASSERT_THAT(IsNotNull(RunGraphNode));
+			})
+			.Then(TEXT("Create a Write Pose node"), [&]()
+			{
+				TArray<UEdGraphPin*> FromPins;
+				WritePoseNode = UAFTestsUtilities::AddUnitNode(ParentGraph, "/Script/AnimNext.RigUnit_AnimNextWriteSkeletalMeshComponentPose", FromPins, FVector2f(336.0f, 0.0f));
+				ASSERT_THAT(IsNotNull(WritePoseNode));
+			})
+			.Then(TEXT("Add link from PrePhysics node Execute pin to Run Graph node Execute pin"), [&]()
+			{
+				ASSERT_THAT(IsTrue(UAFTestsUtilities::AddLink(AnimNextModule, "RigUnit_AnimNextPrePhysicsEvent.ExecuteContext", "AnimNextRunAnimationGraph_v2.ExecuteContext")));
+			})
+			.Then(TEXT("Add link from Run Graph node output pin to Write Pose node Execute pin"), [&]()
+			{				
+				ASSERT_THAT(IsTrue(UAFTestsUtilities::AddLink(AnimNextModule, "AnimNextRunAnimationGraph_v2.ExecuteContext", "AnimNextWriteSkeletalMeshComponentPose.ExecuteContext")));
+			})
+			.Then(TEXT("Collpase Run Graph and Write Pose nodes"), [&]()
+			{
+				TArray<FName> Nodes;
+				Nodes.Add(FName(RunGraphNodeName));
+				Nodes.Add(FName(WritePoseNodeName));
+				URigVMCollapseNode* CollapseNode = UAFTestsUtilities::CollapseNodes(AnimNextModule, Nodes);
+				ASSERT_THAT(IsNotNull(CollapseNode));
+			})
+			.Until(TEXT("Execute Undo Collpase Nodes"), [&]()
+			{
+				return GEditor->UndoTransaction();
+			})
+			.Then(TEXT("Assert Undo Collapse Nodes"), [&]()
+			{
+				ASSERT_THAT(IsNotNull(ControllerGraph->FindNodeByName(FName(RunGraphNodeName))));
+				ASSERT_THAT(IsNotNull(ControllerGraph->FindNodeByName(FName(WritePoseNodeName))));
+				ASSERT_THAT(IsNull(ControllerGraph->FindNodeByName(FName(CollapseNodeName))));
+			})
+			.Until(TEXT("Execute Redo Collapse Nodes"), [&]()
+			{
+				return GEditor->RedoTransaction();
+			})
+			.Then(TEXT("Assert Redo Collapse Nodes"), [&]()
+			{
+				ASSERT_THAT(IsNull(ControllerGraph->FindNodeByName(FName(RunGraphNodeName))));
+				ASSERT_THAT(IsNull(ControllerGraph->FindNodeByName(FName(WritePoseNodeName))));
+				ASSERT_THAT(IsNotNull(ControllerGraph->FindNodeByName(FName(CollapseNodeName))));
+			});
+	}
+	
+protected:
+	void CreateAnimNextModule()
+	{
+		UFactory* Factory = NewObject<UAnimNextModuleFactory>();
+		AnimNextModulePackage = CreateUniquePackage("NewAnimNextModule");
+		AnimNextModulePackageName = *FPaths::GetBaseFilename(AnimNextModulePackage->GetName());		
+		UObject* FactoryObject = Factory->FactoryCreateNew(UAnimNextModule::StaticClass(), AnimNextModulePackage, AnimNextModulePackageName, RF_Public | RF_Standalone, NULL, GWarn);
+		FAssetRegistryModule::AssetCreated(FactoryObject);
+		AnimNextModule = CastChecked<UAnimNextModule>(FactoryObject);
+	}
+	
+	void CreateAnimNextWorkspace()
+	{
+		UFactory* Factory = NewObject<UAnimNextWorkspaceFactory>();
+		AnimNextWorkspacePackage = CreateUniquePackage("NewAnimNextWorkspace");
+		AnimNextWorkspacePackageName = *FPaths::GetBaseFilename(AnimNextWorkspacePackage->GetName());		
+		UObject* FactoryObject = Factory->FactoryCreateNew(UWorkspace::StaticClass(), AnimNextWorkspacePackage, AnimNextWorkspacePackageName, RF_Public | RF_Standalone, NULL, GWarn);
+		FAssetRegistryModule::AssetCreated(FactoryObject);
+		AnimNextWorkspace = CastChecked<UWorkspace>(FactoryObject);
+	}
+	
+	UPackage* CreateUniquePackage(FString InAssetName)
+	{
+		// Create a unique package name and path
+		const FString BasePackageName = "/Game/";    // aka “/All/Content/‘
+		FString UniquePackageName;
+		FString AssetName;
+		IAssetTools& AssetTools = FModuleManager::LoadModuleChecked<FAssetToolsModule>("AssetTools").Get();
+		AssetTools.CreateUniqueAssetName(BasePackageName, InAssetName, UniquePackageName, AssetName);
+		return CreatePackage(*UniquePackageName);
+	}
+};
+
+#endif // WITH_EDITOR
+#endif // WITH_DEV_AUTOMATION_TESTS
diff --git a/Engine/Plugins/Experimental/UAF/UAFTests/Source/UAFTests/Private/UAFTestsModule.cpp b/Engine/Plugins/Experimental/UAF/UAFTests/Source/UAFTests/Private/UAFTestsModule.cpp
new file mode 100644
--- /dev/null
+++ b/Engine/Plugins/Experimental/UAF/UAFTests/Source/UAFTests/Private/UAFTestsModule.cpp
@@ -0,0 +1,20 @@
+// Copyright Epic Games, Inc. All Rights Reserved.
+
+#include "UAFTestsModule.h"
+
+#define LOCTEXT_NAMESPACE "FUAFTestsModule"
+
+void FUAFTestsModule::StartupModule()
+{
+	// This code will execute after your module is loaded into memory; the exact timing is specified in the .uplugin file per-module
+}
+
+void FUAFTestsModule::ShutdownModule()
+{
+	// This function may be called during shutdown to clean up your module.  For modules that support dynamic reloading,
+	// we call this function before unloading the module.
+}
+
+#undef LOCTEXT_NAMESPACE
+
+IMPLEMENT_MODULE(FUAFTestsModule, UAFTests)
diff --git a/Engine/Plugins/Experimental/UAF/UAFTests/Source/UAFTests/Private/UAFTestsUtilities.cpp b/Engine/Plugins/Experimental/UAF/UAFTests/Source/UAFTests/Private/UAFTestsUtilities.cpp
new file mode 100644
--- /dev/null
+++ b/Engine/Plugins/Experimental/UAF/UAFTests/Source/UAFTests/Private/UAFTestsUtilities.cpp
@@ -0,0 +1,160 @@
+// Copyright Epic Games, Inc. All Rights Reserved.
+
+#if WITH_DEV_AUTOMATION_TESTS
+#if WITH_EDITOR
+
+#include "UAFTestsUtilities.h"
+
+namespace UAFTestsUtilities
+{
+	UObject* CreateFactoryObject(UFactory* InFactory, UClass* InClass, const FString& InPackageName)
+	{
+		if (!InClass || !InFactory)
+		{
+			UE_LOG(LogTemp, Error, TEXT("Invalid pointer instance provided."));
+			return nullptr;
+		}
+		
+		UPackage* NewPackage = CreatePackage(nullptr);
+		if (!NewPackage)
+		{
+			UE_LOG(LogTemp, Error, TEXT("Unable to create UPackage instance."));
+			return nullptr;			
+		}
+		
+		FName NewPackageName = InPackageName.Equals("") ? *FPaths::GetBaseFilename(NewPackage->GetName()) : FName(InPackageName);
+		
+		return InFactory->FactoryCreateNew(InClass, NewPackage, NewPackageName, RF_Public | RF_Standalone, NULL, GWarn);
+	}
+
+	UEdGraphNode* AddUnitNode(UEdGraph* InParentGraph, const FString& InScriptStructPath, TArray<UEdGraphPin*>& InFromPins, const FVector2f& InRigUnitLocation)
+	{		
+		if (!InParentGraph)
+		{
+			UE_LOG(LogTemp, Error, TEXT("Invalid UEdGraph instance provided."));
+			return nullptr;
+		}
+		
+		UScriptStruct* ScriptStruct = RigVMTypeUtils::FindObjectFromCPPTypeObjectPath<UScriptStruct>(InScriptStructPath);
+		FText Category = NSLOCTEXT("Category", "Category", "Category");
+		FText MenuDesc = NSLOCTEXT("Menu_Desc", "MenuDesc", "Menu Desc");
+		FText ToolTip = NSLOCTEXT("Tool_Tip", "ToolTip", "Tool Tip");
+		TSharedPtr<FAnimNextSchemaAction_RigUnit> RigUnitAction = MakeShared<FAnimNextSchemaAction_RigUnit>(URigVMUnitNode::StaticClass(), ScriptStruct, Category, MenuDesc, ToolTip);
+
+		return RigUnitAction->PerformAction(InParentGraph, InFromPins, InRigUnitLocation);
+	}
+
+	URigVMLibraryNode* AddFunctionNode(UAnimNextRigVMAsset* InAnimNextAsset, const FString& InFunctionName)
+	{
+		if (!InAnimNextAsset)
+		{
+			UE_LOG(LogTemp, Error, TEXT("Invalid UAnimNextRigVMAsset instance provided."));
+			return nullptr;
+		}
+		
+		UAnimNextRigVMAssetEditorData* EditorData = UE::UAF::UncookedOnly::FUtils::GetEditorData(InAnimNextAsset);
+		return EditorData->AddFunction(FName(InFunctionName), true);
+	}
+
+	UAnimNextVariableEntry* AddVariable(UAnimNextRigVMAsset* InAnimNextAsset, const FAnimNextParamType& InVariableType, const FString& InVariableName, const FString& InDefaultValue)
+	{
+		if (!InAnimNextAsset)
+		{
+			UE_LOG(LogTemp, Error, TEXT("Invalid UAnimNextRigVMAsset instance provided."));
+			return nullptr;
+		}
+		
+		UAnimNextRigVMAssetEditorData* EditorData = UE::UAF::UncookedOnly::FUtils::GetEditorData(InAnimNextAsset);
+		return EditorData->AddVariable(FName(InVariableName), InVariableType, InDefaultValue);
+	}
+
+		UEdGraphNode* AddVariableNode(UEdGraph* InParentGraph, const UObject* InSourceObject, const FString& InVariableName, const FAnimNextParamType& InVariableType, const FAnimNextSchemaAction_Variable::EVariableAccessorChoice InVariableAccessorChoice, TArray<UEdGraphPin*>& InFromPins, const FVector2f& InVariableLocation)
+	{
+		if (!InParentGraph)
+		{
+			UE_LOG(LogTemp, Error, TEXT("Invalid UEdGraph instance provided."));
+			return nullptr;
+		}
+				
+		TSharedPtr<FAnimNextSchemaAction_Variable> VariableAction = MakeShared<FAnimNextSchemaAction_Variable>(FName(InVariableName), InSourceObject, InVariableType, InVariableAccessorChoice);
+		return VariableAction->PerformAction(InParentGraph, InFromPins, InVariableLocation);
+	}
+
+	URigVMPin* AddPin(UAnimNextRigVMAsset* InAnimNextAsset, URigVMLibraryNode* InLibraryNode, ERigVMPinDirection InDirection, const FString& InCPPName, const FString& InCPPType, const FName& InCPPTypeObjectPath, const FString& InDefaultValue)
+	{
+		if (!InAnimNextAsset || !InLibraryNode)
+		{
+			UE_LOG(LogTemp, Error, TEXT("Invalid pointer instance provided."));
+			return nullptr;
+		}
+
+		UAnimNextRigVMAssetEditorData* EditorData = UE::UAF::UncookedOnly::FUtils::GetEditorData(InAnimNextAsset);
+		URigVMGraph* FunctionGraph = InLibraryNode->GetContainedGraph();
+		URigVMController* Controller = CastChecked<IRigVMClientHost>(EditorData)->GetController(FunctionGraph);		
+		const FName NewPinName = Controller->AddExposedPin(FName(InCPPName), InDirection, InCPPType, InCPPTypeObjectPath, InDefaultValue);
+		
+		URigVMPin* NewPin;
+		switch (InDirection)
+		{
+			case ERigVMPinDirection::Input:
+				NewPin = InLibraryNode->GetEntryNode()->FindPin(NewPinName.ToString());			
+				break;
+			case ERigVMPinDirection::Output:
+				NewPin = InLibraryNode->GetReturnNode()->FindPin(NewPinName.ToString());
+				break;
+			default:
+				NewPin = nullptr;
+				break;
+		}
+		
+		return NewPin;
+	}
+
+	bool AddLink(UAnimNextRigVMAsset* InAnimNextAsset, const FString& InOutputPinPath, const FString& InInputPinPath)
+	{
+		if (!InAnimNextAsset)
+		{
+			UE_LOG(LogTemp, Error, TEXT("Invalid pointer instance provided."));
+			return false;
+		}
+		
+		UAnimNextRigVMAssetEditorData* EditorData = UE::UAF::UncookedOnly::FUtils::GetEditorData(InAnimNextAsset);
+		URigVMController* Controller = CastChecked<IRigVMClientHost>(EditorData)->GetControllerByName("RigVMGraph");
+		return Controller->AddLink(InOutputPinPath, InInputPinPath);
+	}
+
+	bool SetNodeSelection(UAnimNextRigVMAsset* InAnimNextAsset, const TArray<FName>& InNodeNames)
+	{
+		if (!InAnimNextAsset)
+		{
+			UE_LOG(LogTemp, Error, TEXT("Invalid pointer instance provided."));
+			return false;
+		}
+		
+		UAnimNextRigVMAssetEditorData* EditorData = UE::UAF::UncookedOnly::FUtils::GetEditorData(InAnimNextAsset);
+		URigVMController* Controller = CastChecked<IRigVMClientHost>(EditorData)->GetControllerByName("RigVMGraph");
+		return Controller->SetNodeSelection(InNodeNames);
+	}
+
+	URigVMCollapseNode* CollapseNodes(UAnimNextRigVMAsset* InAnimNextAsset, const TArray<FName>& InNodeNames, const FString& InCollapseNodeName, bool InCollapseToFunction)
+	{
+		if (SetNodeSelection(InAnimNextAsset, InNodeNames))
+		{
+			UAnimNextRigVMAssetEditorData* EditorData = UE::UAF::UncookedOnly::FUtils::GetEditorData(InAnimNextAsset);
+			URigVMController* Controller = CastChecked<IRigVMClientHost>(EditorData)->GetControllerByName("RigVMGraph");			
+			Controller->OpenUndoBracket(TEXT("Collapse Node"));
+			URigVMCollapseNode* CollapseNode = Controller->CollapseNodes(InNodeNames, InCollapseNodeName);
+			if (InCollapseToFunction)
+			{
+				Controller->PromoteCollapseNodeToFunctionReferenceNode(CollapseNode->GetFName());
+			}
+			Controller->CloseUndoBracket();
+			return CollapseNode;
+		}
+		
+		return nullptr;
+	}
+}
+
+#endif
+#endif
diff --git a/Engine/Plugins/Experimental/UAF/UAFTests/Source/UAFTests/Private/UAFVariableTests.cpp b/Engine/Plugins/Experimental/UAF/UAFTests/Source/UAFTests/Private/UAFVariableTests.cpp
new file mode 100644
--- /dev/null
+++ b/Engine/Plugins/Experimental/UAF/UAFTests/Source/UAFTests/Private/UAFVariableTests.cpp
@@ -0,0 +1,47 @@
+// Copyright Epic Games, Inc. All Rights Reserved.
+
+#if WITH_DEV_AUTOMATION_TESTS
+#if WITH_EDITOR
+
+#include "UAFTestsUtilities.h"
+#include "CQTest.h"
+#include "Editor.h"
+#include "Subsystems/AssetEditorSubsystem.h"
+#include "Graph/AnimNextAnimationGraph.h"
+#include "Graph/AnimNextAnimationGraphFactory.h"
+#include "AnimNextRigVMAssetEditorData.h"
+
+TEST_CLASS(FAssetVariableTests, "Animation.UAF.Functional")
+{
+	UAnimNextRigVMAsset* AnimNextAsset;
+	UAssetEditorSubsystem* AssetEditorSubsystem;
+	UAnimNextRigVMAssetEditorData* EditorData;
+
+	BEFORE_EACH()
+	{
+		AssetEditorSubsystem = GEditor->GetEditorSubsystem<UAssetEditorSubsystem>();
+		ASSERT_THAT(IsNotNull(AssetEditorSubsystem));
+	}
+
+	// QMetry Test Case:  UE-TC-10712
+	TEST_METHOD(AnimNextAnimationGraph_Create_Variable)
+	{		
+		TestCommandBuilder
+			.Do(TEXT("Create AnimNext Animation Graph Asset"), [&]()
+			{
+				const TSubclassOf<UAnimNextAnimationGraphFactory> AnimNextAnimGraphFactoryClass = UAnimNextAnimationGraphFactory::StaticClass();
+				UObject* factoryObject = UAFTestsUtilities::CreateFactoryObject(NewObject<UAnimNextAnimationGraphFactory>(GetTransientPackage(), AnimNextAnimGraphFactoryClass.Get()), UAnimNextAnimationGraph::StaticClass(), "NewAnimNextAnimationGraph");
+				AnimNextAsset = CastChecked<UAnimNextAnimationGraph>(factoryObject);				
+				ASSERT_THAT(IsTrue(AssetEditorSubsystem->OpenEditorForAsset(AnimNextAsset)));
+			})
+			.Then(TEXT("Create Bool Variable"), [&]()
+			{
+				FAnimNextParamType Type = FAnimNextParamType::GetType<bool>();
+				UAnimNextVariableEntry* VariableEntry = UAFTestsUtilities::AddVariable(AnimNextAsset, Type, "NewVariable", "false");
+				ASSERT_THAT(IsNotNull(VariableEntry));
+			});
+	}	
+};
+
+#endif // WITH_EDITOR
+#endif // WITH_DEV_AUTOMATION_TESTS
diff --git a/Engine/Plugins/Experimental/UAF/UAFTests/Source/UAFTests/Private/UAFWorkspaceTests.cpp b/Engine/Plugins/Experimental/UAF/UAFTests/Source/UAFTests/Private/UAFWorkspaceTests.cpp
new file mode 100644
--- /dev/null
+++ b/Engine/Plugins/Experimental/UAF/UAFTests/Source/UAFTests/Private/UAFWorkspaceTests.cpp
@@ -0,0 +1,92 @@
+// Copyright Epic Games, Inc. All Rights Reserved.
+
+#if WITH_DEV_AUTOMATION_TESTS
+#if WITH_EDITOR
+
+#include "UAFTestsUtilities.h"
+#include "CQTest.h"
+#include "Editor.h"
+#include "Subsystems/AssetEditorSubsystem.h"
+#include "WorkspaceEditor/Private/Workspace.h"
+#include "Workspace/AnimNextWorkspaceFactory.h"
+#include "Graph/AnimNextAnimationGraph.h"
+#include "Graph/AnimNextAnimationGraphFactory.h"
+#include "Module/AnimNextModule.h"
+#include "Module/AnimNextModuleFactory.h"
+
+TEST_CLASS(FUAFWorkspaceTests, "Animation.UAF.Functional")
+{
+	UWorkspace* WorkspaceAsset = nullptr;
+	UWorkspace* WorkspaceSecondaryAsset = nullptr;
+	UAnimNextAnimationGraph* AnimationGraphAsset = nullptr;
+	UAnimNextModule* ModuleAsset = nullptr;
+	UAssetEditorSubsystem* AssetEditorSubsystem = nullptr;
+
+	BEFORE_EACH()
+	{
+		AssetEditorSubsystem = GEditor->GetEditorSubsystem<UAssetEditorSubsystem>();
+		ASSERT_THAT(IsNotNull(AssetEditorSubsystem));
+	}
+
+	// QMetry Test Case:  UE-TC-10700
+	TEST_METHOD(Verify_Add_Animation_Graph_To_Workspace)
+	{
+		TestCommandBuilder
+			.Do(TEXT("Create Animation Graph Asset"), [&]()
+			{
+				const TSubclassOf<UAnimNextAnimationGraphFactory> AnimGraphFactoryClass = UAnimNextAnimationGraphFactory::StaticClass();
+				UObject* FactoryObject = UAFTestsUtilities::CreateFactoryObject(NewObject<UAnimNextAnimationGraphFactory>(), UAnimNextAnimationGraph::StaticClass(), "NewAnimNextAnimationGraph");
+				AnimationGraphAsset = CastChecked<UAnimNextAnimationGraph>(FactoryObject);
+				ASSERT_THAT(IsNotNull(AnimationGraphAsset));
+			})
+			.Then(TEXT("Create UAF Workspace Asset"), [&]()
+			{
+				const TSubclassOf<UAnimNextWorkspaceFactory> WorkSpaceFactoryClass = UAnimNextWorkspaceFactory::StaticClass();
+				UObject* FactoryObject = UAFTestsUtilities::CreateFactoryObject(NewObject<UAnimNextWorkspaceFactory>(), UWorkspace::StaticClass(), "NewAnimNextWorkspace");
+				WorkspaceAsset = CastChecked<UWorkspace>(FactoryObject);
+				ASSERT_THAT(IsNotNull(WorkspaceAsset));
+			})
+			.Then(TEXT("Add Animation Graph to Workspace and Open in Editor"), [&]()
+			{
+				ASSERT_THAT(IsTrue(WorkspaceAsset->AddAsset(AnimationGraphAsset)));
+				ASSERT_THAT(IsTrue(AssetEditorSubsystem->OpenEditorForAsset(WorkspaceAsset)));
+			});
+	}
+	
+	// QMetry Test Case:  UE-TC-10716
+	TEST_METHOD(Verify_Add_Module_To_Multiple_Workspaces)
+	{
+		TestCommandBuilder
+			.Do(TEXT("Create UAF Module Asset"), [&]()
+			{
+				const TSubclassOf<UAnimNextModuleFactory> ModuleFactoryClass = UAnimNextModuleFactory::StaticClass();
+				UObject* FactoryObject = UAFTestsUtilities::CreateFactoryObject(NewObject<UAnimNextModuleFactory>(), UAnimNextModule::StaticClass(), "NewAnimNextModule");				
+				ModuleAsset = CastChecked<UAnimNextModule>(FactoryObject);
+				ASSERT_THAT(IsNotNull(ModuleAsset));
+			})
+			.Then(TEXT("Create UAF Workspace Asset"), [&]()
+			{
+				const TSubclassOf<UAnimNextWorkspaceFactory> WorkSpaceFactoryClass = UAnimNextWorkspaceFactory::StaticClass();
+				UObject* FactoryObject = UAFTestsUtilities::CreateFactoryObject(NewObject<UAnimNextWorkspaceFactory>(), UWorkspace::StaticClass(), "NewAnimNextWorkspace");
+				WorkspaceAsset = CastChecked<UWorkspace>(FactoryObject);
+				ASSERT_THAT(IsNotNull(WorkspaceAsset));
+			})
+			.Then(TEXT("Create Another UAF Workspace Asset"), [&]()
+			{
+				const TSubclassOf<UAnimNextWorkspaceFactory> WorkSpaceFactoryClass = UAnimNextWorkspaceFactory::StaticClass();
+				UObject* FactoryObject = UAFTestsUtilities::CreateFactoryObject(NewObject<UAnimNextWorkspaceFactory>(), UWorkspace::StaticClass(), "NewAnimNextWorkspace1");
+				WorkspaceSecondaryAsset = CastChecked<UWorkspace>(FactoryObject);
+				ASSERT_THAT(IsNotNull(WorkspaceSecondaryAsset));
+			})		
+			.Then(TEXT("Add Module to Multiple Workspaces and Open in Editor"), [&]()
+			{
+				ASSERT_THAT(IsTrue(WorkspaceAsset->AddAsset(ModuleAsset)));
+				ASSERT_THAT(IsTrue(WorkspaceSecondaryAsset->AddAsset(ModuleAsset)));
+				ASSERT_THAT(IsTrue(AssetEditorSubsystem->OpenEditorForAsset(WorkspaceAsset)));
+				ASSERT_THAT(IsTrue(AssetEditorSubsystem->OpenEditorForAsset(WorkspaceSecondaryAsset)));
+			});
+	}
+};
+
+#endif // WITH_EDITOR
+#endif // WITH_DEV_AUTOMATION_TESTS
diff --git a/Engine/Plugins/Experimental/UAF/UAFTests/Source/UAFTests/Public/UAFTestsModule.h b/Engine/Plugins/Experimental/UAF/UAFTests/Source/UAFTests/Public/UAFTestsModule.h
new file mode 100644
--- /dev/null
+++ b/Engine/Plugins/Experimental/UAF/UAFTests/Source/UAFTests/Public/UAFTestsModule.h
@@ -0,0 +1,14 @@
+// Copyright Epic Games, Inc. All Rights Reserved.
+
+#pragma once
+
+#include "Modules/ModuleManager.h"
+
+class FUAFTestsModule : public IModuleInterface
+{
+public:
+
+	/** IModuleInterface implementation */
+	virtual void StartupModule() override;
+	virtual void ShutdownModule() override;
+};
diff --git a/Engine/Plugins/Experimental/UAF/UAFTests/Source/UAFTests/Public/UAFTestsUtilities.h b/Engine/Plugins/Experimental/UAF/UAFTests/Source/UAFTests/Public/UAFTestsUtilities.h
new file mode 100644
--- /dev/null
+++ b/Engine/Plugins/Experimental/UAF/UAFTests/Source/UAFTests/Public/UAFTestsUtilities.h
@@ -0,0 +1,119 @@
+// Copyright Epic Games, Inc. All Rights Reserved.
+
+#pragma once
+
+#if WITH_DEV_AUTOMATION_TESTS
+#if WITH_EDITOR
+
+#include "CoreMinimal.h"
+#include "Factories/Factory.h"
+#include "Common/GraphEditorSchemaActions.h"
+#include "AnimNextRigVMAssetEditorData.h"
+#include "Param/ParamType.h"
+
+class URigVMLibraryNode;
+
+namespace UAFTestsUtilities
+{
+	/**
+	 * Creates the object instance for the specified asset class.
+	 *
+	 * @param InFactory  A pointer to the factory instance.
+	 * @param InClass  A pointer to the asset static class.
+	 * @param InPackageName  Name being assigned to the package, optional.
+	 * @return A pointer to the object instance for the specified asset class, or nullptr otherwise.
+	 */
+	UObject* CreateFactoryObject(UFactory* InFactory, UClass* InClass, const FString& InPackageName = TEXT(""));
+
+	/**
+	 * Adds a new RigUnit node to the AnimNext graph specified.
+	 *
+	 * @param InParentGraph  A pointer to the AnimNext parent graph instance.
+	 * @param InScriptStructPath:  Path to the script struct.
+	 * @param InFromPins:  Pin to wire from.
+	 * @param InRigUnitLocation:  Location of the RigUnit node on the graph.
+	 * @return A pointer to the RigUnit node instance, or nullptr otherwise.
+	 */
+	UEdGraphNode* AddUnitNode(UEdGraph* InParentGraph, const FString& InScriptStructPath, TArray<UEdGraphPin*>& InFromPins, const FVector2f& InRigUnitLocation);
+
+	/**
+	 * Adds a new function node instance to the AnimNext Asset specified.
+	 *
+	 * @param InAnimNextAsset  A pointer to the AnimNext asset instance.
+	 * @param InFunctionName  Name to assign to function node, optional.
+	 * @return A pointer to the object instance for the function node, or nullptr otherwise.
+	 */
+	URigVMLibraryNode* AddFunctionNode(UAnimNextRigVMAsset* InAnimNextAsset, const FString& InFunctionName = TEXT("NewFunction"));
+
+	/**
+	 * Adds a new variable instance to the AnimNext Asset specified.
+	 *
+	 * @param InAnimNextAsset  A pointer to the AnimNext asset instance.
+	 * @param InVariableType:  Data type of variable.
+	 * @param InVariableName  Name to assign to variable, optional.
+	 * @param InDefaultValue  Default value to assign to variable, optional.
+	 * @return A pointer to the object instance for the variable, or nullptr otherwise.
+	 */
+	UAnimNextVariableEntry* AddVariable(UAnimNextRigVMAsset* InAnimNextAsset, const FAnimNextParamType& InVariableType, const FString& InVariableName = TEXT("NewVariable"), const FString& InDefaultValue = TEXT(""));
+
+	/**
+	 * Adds a new variable node to the AnimNext graph specified.
+	 *
+	 * @param InParentGraph  A pointer to the AnimNext parent graph instance.
+	 * @param InSourceObject The source object (asset or struct) that this variable comes from
+	 * @param InVariableName  Name to assign to variable.
+	 * @param InVariableType:  Data type of variable.
+	 * @param InVariableAccessorChoice:  Accessor choice of the variable node.
+	 * @param InFromPins:  Pin to wire from.
+	 * @param InVariableLocation:  Location of the variable node on the graph.
+	 * @return A pointer to the variable node instance, or nullptr otherwise.
+	 */
+	UEdGraphNode* AddVariableNode(UEdGraph* InParentGraph, const UObject* InSourceObject, const FString& InVariableName, const FAnimNextParamType& InVariableType, const FAnimNextSchemaAction_Variable::EVariableAccessorChoice InVariableAccessorChoice, TArray<UEdGraphPin*>& InFromPins, const FVector2f& InVariableLocation);
+
+	/**
+	 * Adds a new pin to the library node specified.
+	 *
+	 * @param InAnimNextAsset  A pointer to the AnimNext asset instance.
+	 * @param InLibraryNode  A pointer to the library node instance.
+	 * @param InDirection  Direction of the pin.
+	 * @param InCPPName  Name assigned to the pin, optional.
+	 * @param InCPPType  Data type assigned to the pin, optional.
+	 * @param InCPPTypeObjectPath  Path assigned to the pin, optional.
+	 * @param InDefaultValue  Default value assigned to the pin, optional.
+	 * @return A pointer to the pin instance for the specified library node, or nullptr otherwise.
+	 */
+	URigVMPin* AddPin(UAnimNextRigVMAsset* InAnimNextAsset, URigVMLibraryNode* InLibraryNode, ERigVMPinDirection InDirection, const FString& InCPPName = TEXT("Argument"), const FString& InCPPType = TEXT("bool"), const FName& InCPPTypeObjectPath = TEXT("None"), const FString& InDefaultValue = TEXT("False"));
+
+	/**
+	 * Adds a new link between two pins.
+	 *
+	 * @param InAnimNextAsset  A pointer to the AnimNext asset instance.
+	 * @param InOutputPinPath  A path to an output pin.
+	 * @param InInputPinPath  A path to an input pin.
+	 * @return True if add link is successful, False otherwise.
+	 */
+	bool AddLink(UAnimNextRigVMAsset* InAnimNextAsset, const FString& InOutputPinPath, const FString& InInputPinPath);
+
+	/**
+	 * Select specified node(s).
+	 *
+	 * @param InAnimNextAsset  A pointer to the AnimNext asset instance.
+	 * @param InNodeNames  Array of node name(s) to select.
+	 * @return True if node(s) selection is successful, False otherwise.
+	 */
+	bool SetNodeSelection(UAnimNextRigVMAsset* InAnimNextAsset, const TArray<FName>& InNodeNames);
+
+	/**
+	 * Collapse selected nodes to a single node.
+	 *
+	 * @param InAnimNextAsset  A pointer to the AnimNext asset instance.
+	 * @param InNodeNames  Array of node names to collapse.
+	 * @param InCollapseNodeName  Name of collapse node, optional
+	 * @param InCollapseToFunction  Collapse to a function node, optional.
+	 * @return A pointer to the pin instance for the collapse node, or nullptr otherwise.
+	 */
+	URigVMCollapseNode* CollapseNodes(UAnimNextRigVMAsset* InAnimNextAsset, const TArray<FName>& InNodeNames, const FString& InCollapseNodeName = TEXT(""), bool InCollapseToFunction = false);
+}
+
+#endif // WITH_EDITOR
+#endif // WITH_DEV_AUTOMATION_TESTS
diff --git a/Engine/Plugins/Experimental/UAF/UAFTests/Source/UAFTests/UAFTests.Build.cs b/Engine/Plugins/Experimental/UAF/UAFTests/Source/UAFTests/UAFTests.Build.cs
new file mode 100644
--- /dev/null
+++ b/Engine/Plugins/Experimental/UAF/UAFTests/Source/UAFTests/UAFTests.Build.cs
@@ -0,0 +1,63 @@
+// Copyright Epic Games, Inc. All Rights Reserved.
+
+using System.IO;
+using UnrealBuildTool;
+
+public class UAFTests : ModuleRules
+{
+	public UAFTests(ReadOnlyTargetRules Target) : base(Target)
+	{
+		PCHUsage = ModuleRules.PCHUsageMode.UseExplicitOrSharedPCHs;
+		bTreatAsEngineModule = true;
+
+		PublicIncludePaths.AddRange(
+			new string[] {
+				// ... add public include paths required here ...
+			}
+		);
+				
+		
+		PrivateIncludePaths.AddRange(
+			new string[] {
+				// ... add other private include paths required here ...
+			}
+		);
+			
+		
+		PublicDependencyModuleNames.AddRange(
+			new string[]
+			{
+				"Core",
+				"UnrealEd"
+				// ... add other public dependencies that you statically link with here ...
+			}
+		);
+			
+		
+		PrivateDependencyModuleNames.AddRange(
+			new string[]
+			{
+				"CoreUObject",
+				"CQTest",
+				"Engine",
+				"UAF",
+				"WorkspaceEditor",
+				"UAFEditor",
+				"UAFStateTree",
+				"UAFStateTreeEditor",
+				"UAFAnimGraph",
+				"UAFAnimGraphEditor",
+				"UAFUncookedOnly",
+				"RigVMDeveloper"
+				// ... add private dependencies that you statically link with here ...	
+			}
+		);
+
+		DynamicallyLoadedModuleNames.AddRange(
+			new string[]
+			{
+				// ... add any modules that your module loads dynamically here ...
+			}
+		);
+	}
+}
diff --git a/Engine/Plugins/Experimental/UAF/UAFTests/UAFTests.uplugin b/Engine/Plugins/Experimental/UAF/UAFTests/UAFTests.uplugin
new file mode 100644
--- /dev/null
+++ b/Engine/Plugins/Experimental/UAF/UAFTests/UAFTests.uplugin
@@ -0,0 +1,46 @@
+{
+	"FileVersion": 1,
+	"Version": 1,
+	"VersionName": "1.0",
+	"FriendlyName": "UAF Tests",
+	"Description": "UAF Automated Tests",
+	"Category": "Other",
+	"CreatedBy": "Epic Games, Inc.",
+	"CreatedByURL": "https://epicgames.com",
+	"DocsURL": "",
+	"MarketplaceURL": "",
+	"SupportURL": "",
+	"CanContainContent": true,
+	"IsBetaVersion": false,
+	"IsExperimentalVersion": true,
+	"Installed": false,
+	"Modules": [
+		{
+			"Name": "UAFTests",
+			"Type": "Editor",
+			"LoadingPhase": "Default"
+		}
+	],
+	"Plugins": [		
+		{
+			"Name": "Workspace",
+			"Enabled": true
+		},
+		{
+			"Name": "UAF",
+			"Enabled": true
+		},
+		{
+			"Name": "UAFAnimGraph",
+			"Enabled": true		
+		},		
+		{
+			"Name": "UAFStateTree",
+			"Enabled": true		
+		},
+		{
+			"Name": "RigVM",
+			"Enabled": true		
+		}
+	]
+}
