Move UAF tests from Engine into the UAFTest plugin

diff --git a/Engine/Plugins/Experimental/UAF/UAF/UAF.uplugin b/Engine/Plugins/Experimental/UAF/UAF/UAF.uplugin
--- a/Engine/Plugins/Experimental/UAF/UAF/UAF.uplugin
+++ b/Engine/Plugins/Experimental/UAF/UAF/UAF.uplugin
@@ -67,11 +67,6 @@
 			"Name" : "UAFUncookedOnly",
 			"Type" : "UncookedOnly",
 			"LoadingPhase" : "Default"
-		},
-		{
-			"Name" : "UAFTestSuite",
-			"Type" : "UncookedOnly",
-			"LoadingPhase" : "Default"
 		}
 	]
 }
diff --git a/Engine/Plugins/Experimental/UAF/UAFAnimGraph/UAFAnimGraph.uplugin b/Engine/Plugins/Experimental/UAF/UAFAnimGraph/UAFAnimGraph.uplugin
--- a/Engine/Plugins/Experimental/UAF/UAFAnimGraph/UAFAnimGraph.uplugin
+++ b/Engine/Plugins/Experimental/UAF/UAFAnimGraph/UAFAnimGraph.uplugin
@@ -42,11 +42,6 @@
 			"Name" : "UAFAnimGraphUncookedOnly",
 			"Type" : "UncookedOnly",
 			"LoadingPhase" : "Default"
-		},
-		{
-			"Name" : "UAFAnimGraphTestSuite",
-			"Type" : "UncookedOnly",
-			"LoadingPhase" : "Default"
 		}
 	]
 }
diff --git a/Engine/Plugins/Experimental/UAF/UAFTests/UAFTests.uplugin b/Engine/Plugins/Experimental/UAF/UAFTests/UAFTests.uplugin
--- a/Engine/Plugins/Experimental/UAF/UAFTests/UAFTests.uplugin
+++ b/Engine/Plugins/Experimental/UAF/UAFTests/UAFTests.uplugin
@@ -19,6 +19,16 @@
 			"Name": "UAFTests",
 			"Type": "Editor",
 			"LoadingPhase": "Default"
+		},
+		{
+			"Name" : "UAFTestSuite",
+			"Type" : "UncookedOnly",
+			"LoadingPhase" : "Default"
+		},
+		{
+			"Name" : "UAFAnimGraphTestSuite",
+			"Type" : "UncookedOnly",
+			"LoadingPhase" : "Default"
 		}
 	],
 	"Plugins": [		
