Fix array OOB in UAF tests

#rnx

diff --git a/Engine/Plugins/Experimental/UAF/UAFTests/Source/UAFAnimGraphTestSuite/Private/AnimNextAnimGraphTraitGraphTest.cpp b/Engine/Plugins/Experimental/UAF/UAFTests/Source/UAFAnimGraphTestSuite/Private/AnimNextAnimGraphTraitGraphTest.cpp
--- a/Engine/Plugins/Experimental/UAF/UAFTests/Source/UAFAnimGraphTestSuite/Private/AnimNextAnimGraphTraitGraphTest.cpp
+++ b/Engine/Plugins/Experimental/UAF/UAFTests/Source/UAFAnimGraphTestSuite/Private/AnimNextAnimGraphTraitGraphTest.cpp
@@ -502,7 +502,7 @@ bool FAnimationAnimNextRuntimeTest_GraphExecute::RunTest(const FString& InParame
 			(void)UE::UAF::EvaluateGraph(EvaluateGraphContext);
 		}
 
-		AddErrorIfFalse(Messages.Num() == 9, "FAnimationAnimNextRuntimeTest_GraphExecute - > Unexpected message count");
+		UE_RETURN_ON_ERROR(Messages.Num() == 9, "FAnimationAnimNextRuntimeTest_GraphExecute - > Unexpected message count");
 		AddErrorIfFalse(Messages[0] == "UpdateCount == 1", "FAnimationAnimNextRuntimeTest_GraphExecute -> Unexpected update count");
 		AddErrorIfFalse(Messages[1] == "SomeInt32 == 78", "FAnimationAnimNextRuntimeTest_GraphExecute -> Unexpected SomeInt32 value");
 		AddErrorIfFalse(Messages[2] == "SomeFloat == 142.33", "FAnimationAnimNextRuntimeTest_GraphExecute -> Unexpected SomeFloat value");
@@ -650,7 +650,7 @@ bool FAnimationAnimNextRuntimeTest_GraphExecuteLatent::RunTest(const FString& In
 			(void)UE::UAF::EvaluateGraph(EvaluateGraphContext);
 		}
 
-		AddErrorIfFalse(Messages.Num() == 9, "FAnimationAnimNextRuntimeTest_GraphExecuteLatent - > Unexpected message count");
+		UE_RETURN_ON_ERROR(Messages.Num() == 9, "FAnimationAnimNextRuntimeTest_GraphExecuteLatent - > Unexpected message count");
 		AddErrorIfFalse(Messages[0] == "UpdateCount == 1", "FAnimationAnimNextRuntimeTest_GraphExecuteLatent -> Unexpected update count");
 		AddErrorIfFalse(Messages[1] == "SomeInt32 == 78", "FAnimationAnimNextRuntimeTest_GraphExecuteLatent -> Unexpected SomeInt32 value");
 		AddErrorIfFalse(Messages[2] == "SomeFloat == 142.33", "FAnimationAnimNextRuntimeTest_GraphExecuteLatent -> Unexpected SomeFloat value");
