UAF: Fix variables tests broken by CL 46470011

diff --git a/Engine/Plugins/Experimental/UAF/UAF/Source/UAF/Internal/Module/AnimNextModuleInstance.h b/Engine/Plugins/Experimental/UAF/UAF/Source/UAF/Internal/Module/AnimNextModuleInstance.h
--- a/Engine/Plugins/Experimental/UAF/UAF/Source/UAF/Internal/Module/AnimNextModuleInstance.h
+++ b/Engine/Plugins/Experimental/UAF/UAF/Source/UAF/Internal/Module/AnimNextModuleInstance.h
@@ -45,6 +45,7 @@ namespace UE::UAF
 namespace UE::UAF::Tests
 {
 	class FVariables_Shared;
+	class FVariables_UOLBindings;
 }
 
 // Root memory owner of a parameterized schedule 
@@ -181,6 +182,7 @@ private:
 	friend FAnimNextNotifyDispatcherComponent;
 	friend FAnimNextGraphInstance;
 	friend UE::UAF::Tests::FVariables_Shared;
+	friend UE::UAF::Tests::FVariables_UOLBindings;
 
 	// Object this entry is bound to
 	UPROPERTY(Transient)
diff --git a/Engine/Plugins/Experimental/UAF/UAFTests/Source/UAFTestSuite/Private/AnimNextVariablesTest.cpp b/Engine/Plugins/Experimental/UAF/UAFTests/Source/UAFTestSuite/Private/AnimNextVariablesTest.cpp
--- a/Engine/Plugins/Experimental/UAF/UAFTests/Source/UAFTestSuite/Private/AnimNextVariablesTest.cpp
+++ b/Engine/Plugins/Experimental/UAF/UAFTests/Source/UAFTestSuite/Private/AnimNextVariablesTest.cpp
@@ -25,6 +25,8 @@
 
 // AnimNext Parameters Tests
 
+#include "Factory/UAFSystemFactoryParams.h"
+
 #include UE_INLINE_GENERATED_CPP_BY_NAME(AnimNextVariablesTest)
 
 #if WITH_DEV_AUTOMATION_TESTS
@@ -441,6 +443,7 @@ bool FVariables_UOLBindings::RunTest(const FString& InParameters)
 	Asset->GetRigVMExtendedExecuteContext().SetRuntimeSettings(RuntimeSettings);
 
 	FAnimNextModuleInstance Instance(Asset, NewObject<UAnimNextComponent>(), EAnimNextModuleInitMethod::None);
+	Instance.Initialize(FUAFSystemFactoryParams());
 	FUAFRigVMComponent& RigVMComponent = Instance.GetComponent<FUAFRigVMComponent>();
 
 	// Execute bindings to fetch data
@@ -515,6 +518,7 @@ bool FVariables_Shared::RunTest(const FString& InParameters)
 	UE_RETURN_ON_ERROR(Entry != nullptr, "FVariables_Shared -> Could not add shared variables struct.");
 
 	FAnimNextModuleInstance Instance(Asset, EAnimNextModuleInitMethod::None);
+	Instance.Initialize(FUAFSystemFactoryParams());
 
 	Private::FTestReadWrite::Test<bool>(*this, Instance, TEXT("bBool"), true);
 	Private::FTestReadWrite::Test<bool>(*this, Instance, TEXT("bBitfield0"), true);
