UAF: Disable documentation policy for anim next test suite

#rnx

diff --git a/Engine/Plugins/Experimental/UAF/UAFTests/Source/UAFTestSuite/Private/AnimNextVariablesTest.h b/Engine/Plugins/Experimental/UAF/UAFTests/Source/UAFTestSuite/Private/AnimNextVariablesTest.h
--- a/Engine/Plugins/Experimental/UAF/UAFTests/Source/UAFTestSuite/Private/AnimNextVariablesTest.h
+++ b/Engine/Plugins/Experimental/UAF/UAFTests/Source/UAFTestSuite/Private/AnimNextVariablesTest.h
@@ -139,8 +139,14 @@ struct FAnimNextParamTypeTestStruct
 //	TArray<TSoftClassPtr<UObject>> SoftClassPtrArray;
 };
 
-USTRUCT(meta=(Hidden))
-struct FAnimNextTests_TestOperation : public FRigVMStructMutable
+USTRUCT(meta=(Hidden, DocumentationPolicy="None"))
+struct FAnimNextTests_TestBase : public FRigVMStructMutable
+{
+	GENERATED_BODY()
+};
+
+USTRUCT()
+struct FAnimNextTests_TestOperation : public FAnimNextTests_TestBase
 {
 	GENERATED_BODY()
 
@@ -157,8 +163,8 @@ struct FAnimNextTests_TestOperation : public FRigVMStructMutable
 	int32 Result = 0;
 };
 
-USTRUCT(meta=(Hidden))
-struct FAnimNextTests_PrintResult : public FRigVMStructMutable
+USTRUCT()
+struct FAnimNextTests_PrintResult : public FAnimNextTests_TestBase
 {
 	GENERATED_BODY()
 
