Remove manual RegisterWithGC/UnregisterWithGC in the IGarbageCollection trait interface, replaced with automatic registration

diff --git a/Engine/Plugins/Experimental/UAF/UAFAnimGraph/Source/UAFAnimGraph/Internal/Injection/InjectionSiteTrait.h b/Engine/Plugins/Experimental/UAF/UAFAnimGraph/Source/UAFAnimGraph/Internal/Injection/InjectionSiteTrait.h
--- a/Engine/Plugins/Experimental/UAF/UAFAnimGraph/Source/UAFAnimGraph/Internal/Injection/InjectionSiteTrait.h
+++ b/Engine/Plugins/Experimental/UAF/UAFAnimGraph/Source/UAFAnimGraph/Internal/Injection/InjectionSiteTrait.h
@@ -146,7 +146,6 @@ namespace UE::UAF
 			FAnimNextAnimGraph CachedGraph;
 
 			void Construct(const FExecutionContext& Context, const FTraitBinding& Binding);
-			void Destruct(const FExecutionContext& Context, const FTraitBinding& Binding);
 		};
 
 		// Internal impl - ensure that we have a pending request that we have dispatched to the
diff --git a/Engine/Plugins/Experimental/UAF/UAFAnimGraph/Source/UAFAnimGraph/Internal/Traits/BlendStackTrait.h b/Engine/Plugins/Experimental/UAF/UAFAnimGraph/Source/UAFAnimGraph/Internal/Traits/BlendStackTrait.h
--- a/Engine/Plugins/Experimental/UAF/UAFAnimGraph/Source/UAFAnimGraph/Internal/Traits/BlendStackTrait.h
+++ b/Engine/Plugins/Experimental/UAF/UAFAnimGraph/Source/UAFAnimGraph/Internal/Traits/BlendStackTrait.h
@@ -157,9 +157,6 @@ namespace UE::UAF
 			// The index of the currently active graph
 			// All other graphs are blending out
 			int32 CurrentlyActiveGraphIndex = INDEX_NONE;
-
-			void Construct(const FExecutionContext& Context, const FTraitBinding& Binding);
-			void Destruct(const FExecutionContext& Context, const FTraitBinding& Binding);
 		};
 
 		using FSharedData = FAnimNextBlendStackCoreTraitSharedData;
diff --git a/Engine/Plugins/Experimental/UAF/UAFAnimGraph/Source/UAFAnimGraph/Private/Injection/InjectionSiteTrait.cpp b/Engine/Plugins/Experimental/UAF/UAFAnimGraph/Source/UAFAnimGraph/Private/Injection/InjectionSiteTrait.cpp
--- a/Engine/Plugins/Experimental/UAF/UAFAnimGraph/Source/UAFAnimGraph/Private/Injection/InjectionSiteTrait.cpp
+++ b/Engine/Plugins/Experimental/UAF/UAFAnimGraph/Source/UAFAnimGraph/Private/Injection/InjectionSiteTrait.cpp
@@ -65,15 +65,6 @@ namespace UE::UAF
 		{
 			(void)ModuleInstance->GetComponent<FAnimNextModuleInjectionComponent>();
 		}
-
-		IGarbageCollection::RegisterWithGC(Context, Binding);
-	}
-
-	void FInjectionSiteTrait::FInstanceData::Destruct(const FExecutionContext& Context, const FTraitBinding& Binding)
-	{
-		FTrait::FInstanceData::Destruct(Context, Binding);
-
-		IGarbageCollection::UnregisterWithGC(Context, Binding);
 	}
 
 	void FInjectionSiteTrait::TrackPendingInjectionRequest(int32 InChildIndex, FInstanceData& InstanceData)
diff --git a/Engine/Plugins/Experimental/UAF/UAFAnimGraph/Source/UAFAnimGraph/Private/PlayAnim/PlayAnimSlotTrait.cpp b/Engine/Plugins/Experimental/UAF/UAFAnimGraph/Source/UAFAnimGraph/Private/PlayAnim/PlayAnimSlotTrait.cpp
--- a/Engine/Plugins/Experimental/UAF/UAFAnimGraph/Source/UAFAnimGraph/Private/PlayAnim/PlayAnimSlotTrait.cpp
+++ b/Engine/Plugins/Experimental/UAF/UAFAnimGraph/Source/UAFAnimGraph/Private/PlayAnim/PlayAnimSlotTrait.cpp
@@ -48,20 +48,6 @@ namespace UE::UAF
 		bWasRelevant = false;
 	}
 
-	void FPlayAnimSlotTrait::FInstanceData::Construct(const FExecutionContext& Context, const FTraitBinding& Binding)
-	{
-		FTrait::FInstanceData::Construct(Context, Binding);
-
-		IGarbageCollection::RegisterWithGC(Context, Binding);
-	}
-
-	void FPlayAnimSlotTrait::FInstanceData::Destruct(const FExecutionContext& Context, const FTraitBinding& Binding)
-	{
-		FTrait::FInstanceData::Destruct(Context, Binding);
-
-		IGarbageCollection::UnregisterWithGC(Context, Binding);
-	}
-
 	int32 FPlayAnimSlotTrait::FindFreeRequestIndexOrAdd(FInstanceData& InstanceData)
 	{
 		// Find an empty request we can use
diff --git a/Engine/Plugins/Experimental/UAF/UAFAnimGraph/Source/UAFAnimGraph/Private/PlayAnim/PlayAnimSlotTrait.h b/Engine/Plugins/Experimental/UAF/UAFAnimGraph/Source/UAFAnimGraph/Private/PlayAnim/PlayAnimSlotTrait.h
--- a/Engine/Plugins/Experimental/UAF/UAFAnimGraph/Source/UAFAnimGraph/Private/PlayAnim/PlayAnimSlotTrait.h
+++ b/Engine/Plugins/Experimental/UAF/UAFAnimGraph/Source/UAFAnimGraph/Private/PlayAnim/PlayAnimSlotTrait.h
@@ -137,9 +137,6 @@ namespace UE::UAF
 			// The index of the currently active request
 			// All other requests are blending out
 			int32 CurrentlyActiveRequestIndex = INDEX_NONE;
-
-			void Construct(const FExecutionContext& Context, const FTraitBinding& Binding);
-			void Destruct(const FExecutionContext& Context, const FTraitBinding& Binding);
 		};
 
 		// Internal impl
diff --git a/Engine/Plugins/Experimental/UAF/UAFAnimGraph/Source/UAFAnimGraph/Private/TraitCore/ExecutionContext.cpp b/Engine/Plugins/Experimental/UAF/UAFAnimGraph/Source/UAFAnimGraph/Private/TraitCore/ExecutionContext.cpp
--- a/Engine/Plugins/Experimental/UAF/UAFAnimGraph/Source/UAFAnimGraph/Private/TraitCore/ExecutionContext.cpp
+++ b/Engine/Plugins/Experimental/UAF/UAFAnimGraph/Source/UAFAnimGraph/Private/TraitCore/ExecutionContext.cpp
@@ -9,6 +9,7 @@
 #include "TraitCore/NodeInstance.h"
 #include "TraitCore/NodeTemplate.h"
 #include "TraitCore/NodeTemplateRegistry.h"
+#include "TraitInterfaces/IGarbageCollection.h"
 #include "Graph/AnimNextAnimationGraph.h"
 #include "Graph/AnimNextGraphInstance.h"
 #include "Components/SceneComponent.h"
@@ -185,8 +186,17 @@ namespace UE::UAF
 				StackBinding.BaseTraitIndex = TraitIndex;
 				StackBinding.TopTraitIndex = TraitDesc->GetNumStackTraits() - 1;
 			}
-
+			
 			const FTraitBinding Binding(&StackBinding, Trait, TraitIndex);
+			
+			if (TraitDesc->IsGCAware())
+			{
+				const IGarbageCollection* GCImpl = static_cast<const IGarbageCollection*>(Trait->GetTraitInterface(IGarbageCollection::InterfaceUID));
+				ensure(GCImpl);
+				
+				GCImpl->RegisterWithGC(*this, Binding);
+			}
+
 			Trait->ConstructTraitInstance(*this, Binding);
 		}
 
@@ -273,6 +283,17 @@ namespace UE::UAF
 			}
 
 			const FTraitBinding Binding(&StackBinding, Trait, TraitIndex);
+			
+			// Auto-register this trait with garbage collection
+			if (TraitDesc->IsGCAware())
+			{
+				const IGarbageCollection* GCImpl = static_cast<const IGarbageCollection*>(Trait->GetTraitInterface(IGarbageCollection::InterfaceUID));
+				if (ensure(GCImpl))
+				{
+					GCImpl->RegisterWithGC(*this, Binding);
+				}
+			}
+            			
 			Trait->ConstructTraitInstance(*this, Binding);
 		}
 
@@ -333,9 +354,20 @@ namespace UE::UAF
 			const FTraitTemplate* BaseTraitDesc = TraitDesc - TraitDesc->GetTraitIndex();
 			StackBinding.BaseTraitIndex = BaseTraitDesc - TraitDescs;
 			StackBinding.TopTraitIndex = StackBinding.BaseTraitIndex + BaseTraitDesc->GetNumStackTraits() - 1;
-
+			
 			const uint32 TraitIndex = TraitDesc - TraitDescs;
 			const FTraitBinding Binding(&StackBinding, Trait, TraitIndex);
+			
+			// Auto-unregister this trait with garbage collection
+			if (TraitDesc->IsGCAware())
+			{
+				const IGarbageCollection* GCImpl = static_cast<const IGarbageCollection*>(Trait->GetTraitInterface(IGarbageCollection::InterfaceUID));
+				if (ensure(GCImpl))
+				{
+					GCImpl->UnregisterWithGC(*this, Binding);
+				}
+			}
+			
 			Trait->DestructTraitInstance(*this, Binding);
 		}
 
diff --git a/Engine/Plugins/Experimental/UAF/UAFAnimGraph/Source/UAFAnimGraph/Private/TraitCore/NodeTemplateBuilder.cpp b/Engine/Plugins/Experimental/UAF/UAFAnimGraph/Source/UAFAnimGraph/Private/TraitCore/NodeTemplateBuilder.cpp
--- a/Engine/Plugins/Experimental/UAF/UAFAnimGraph/Source/UAFAnimGraph/Private/TraitCore/NodeTemplateBuilder.cpp
+++ b/Engine/Plugins/Experimental/UAF/UAFAnimGraph/Source/UAFAnimGraph/Private/TraitCore/NodeTemplateBuilder.cpp
@@ -5,6 +5,7 @@
 #include "TraitCore/Trait.h"
 #include "TraitCore/TraitRegistry.h"
 #include "TraitCore/NodeTemplate.h"
+#include "TraitInterfaces/IGarbageCollection.h"
 
 namespace UE::UAF
 {
@@ -103,9 +104,11 @@ namespace UE::UAF
 				TraitIndexOrNumTraits++;
 			}
 		}
+		
+		const bool bGCAware = Trait->GetTraitInterface(IGarbageCollection::InterfaceUID) != nullptr;
 
 		// Append our trait template
 		const int32 BufferIndex = NodeTemplateBuffer.AddUninitialized(sizeof(FTraitTemplate));
-		new(&NodeTemplateBuffer[BufferIndex]) FTraitTemplate(TraitUID, TraitHandle, TraitMode, TraitIndexOrNumTraits);
+		new(&NodeTemplateBuffer[BufferIndex]) FTraitTemplate(TraitUID, TraitHandle, TraitMode, bGCAware, TraitIndexOrNumTraits);
 	}
 }
diff --git a/Engine/Plugins/Experimental/UAF/UAFAnimGraph/Source/UAFAnimGraph/Private/TraitCore/TraitTemplate.cpp b/Engine/Plugins/Experimental/UAF/UAFAnimGraph/Source/UAFAnimGraph/Private/TraitCore/TraitTemplate.cpp
--- a/Engine/Plugins/Experimental/UAF/UAFAnimGraph/Source/UAFAnimGraph/Private/TraitCore/TraitTemplate.cpp
+++ b/Engine/Plugins/Experimental/UAF/UAFAnimGraph/Source/UAFAnimGraph/Private/TraitCore/TraitTemplate.cpp
@@ -47,7 +47,22 @@ namespace UE::UAF
 			}
 		}
 
-		Ar << Mode;
+		uint8 ModeByte = Mode;
+		Ar << ModeByte;
+
+		if (Ar.IsLoading())
+		{
+			Mode = ModeByte;
+		}
+		
+		uint8 GCAwareByte = bGCAware;
+		Ar << GCAwareByte;
+
+		if (Ar.IsLoading())
+		{
+			bGCAware = GCAwareByte;
+		}
+		
 		Ar << TraitIndexOrNumTraits;
 	}
 }
diff --git a/Engine/Plugins/Experimental/UAF/UAFAnimGraph/Source/UAFAnimGraph/Private/Traits/BlendSpacePlayer.cpp b/Engine/Plugins/Experimental/UAF/UAFAnimGraph/Source/UAFAnimGraph/Private/Traits/BlendSpacePlayer.cpp
--- a/Engine/Plugins/Experimental/UAF/UAFAnimGraph/Source/UAFAnimGraph/Private/Traits/BlendSpacePlayer.cpp
+++ b/Engine/Plugins/Experimental/UAF/UAFAnimGraph/Source/UAFAnimGraph/Private/Traits/BlendSpacePlayer.cpp
@@ -30,22 +30,7 @@ namespace UE::UAF
 
 	GENERATE_ANIM_TRAIT_IMPLEMENTATION(FBlendSpacePlayerTrait, TRAIT_INTERFACE_ENUMERATOR, NULL_ANIM_TRAIT_INTERFACE_ENUMERATOR, NULL_ANIM_TRAIT_EVENT_ENUMERATOR)
 	#undef TRAIT_INTERFACE_ENUMERATOR
-
-	void FBlendSpacePlayerTrait::FInstanceData::Construct(const FExecutionContext& Context, const FTraitBinding& Binding)
-	{
-		FTrait::FInstanceData::Construct(Context, Binding);
-
-		IGarbageCollection::RegisterWithGC(Context, Binding);
-
-	}
-
-	void FBlendSpacePlayerTrait::FInstanceData::Destruct(const FExecutionContext& Context, const FTraitBinding& Binding)
-	{
-		FTrait::FInstanceData::Destruct(Context, Binding);
-
-		IGarbageCollection::UnregisterWithGC(Context, Binding);
-	}
-
+	
 	void FBlendSpacePlayerTrait::PostEvaluate(FEvaluateTraversalContext& Context, const TTraitBinding<IEvaluate>& Binding) const
 	{
 		const FInstanceData* InstanceData = Binding.GetInstanceData<FInstanceData>();
diff --git a/Engine/Plugins/Experimental/UAF/UAFAnimGraph/Source/UAFAnimGraph/Private/Traits/BlendSpacePlayer.h b/Engine/Plugins/Experimental/UAF/UAFAnimGraph/Source/UAFAnimGraph/Private/Traits/BlendSpacePlayer.h
--- a/Engine/Plugins/Experimental/UAF/UAFAnimGraph/Source/UAFAnimGraph/Private/Traits/BlendSpacePlayer.h
+++ b/Engine/Plugins/Experimental/UAF/UAFAnimGraph/Source/UAFAnimGraph/Private/Traits/BlendSpacePlayer.h
@@ -66,9 +66,6 @@ namespace UE::UAF
 			FBlendFilter BlendFilter;
 
 			int32 NumChildren = 0;
-
-			void Construct(const FExecutionContext& Context, const FTraitBinding& Binding);
-			void Destruct(const FExecutionContext& Context, const FTraitBinding& Binding);
 		};
 
 		// IEvaluate impl
diff --git a/Engine/Plugins/Experimental/UAF/UAFAnimGraph/Source/UAFAnimGraph/Private/Traits/BlendStackTrait.cpp b/Engine/Plugins/Experimental/UAF/UAFAnimGraph/Source/UAFAnimGraph/Private/Traits/BlendStackTrait.cpp
--- a/Engine/Plugins/Experimental/UAF/UAFAnimGraph/Source/UAFAnimGraph/Private/Traits/BlendStackTrait.cpp
+++ b/Engine/Plugins/Experimental/UAF/UAFAnimGraph/Source/UAFAnimGraph/Private/Traits/BlendStackTrait.cpp
@@ -55,19 +55,7 @@ namespace UE::UAF
 #undef TRAIT_REQUIRED_INTERFACE_ENUMERATOR
 
 
-
-	void FBlendStackCoreTrait::FInstanceData::Construct(const FExecutionContext& Context, const FTraitBinding& Binding)
-	{
-		FTrait::FInstanceData::Construct(Context, Binding);
-		IGarbageCollection::RegisterWithGC(Context, Binding);
-	}
-
-	void FBlendStackCoreTrait::FInstanceData::Destruct(const FExecutionContext& Context, const FTraitBinding& Binding)
-	{
-		FTrait::FInstanceData::Destruct(Context, Binding);
-		IGarbageCollection::UnregisterWithGC(Context, Binding);
-	}
-
+	
 	void FBlendStackCoreTrait::FGraphState::Initialize(IBlendStack::FGraphRequest&& GraphRequest)
 	{
 		Request = MoveTemp(GraphRequest);
diff --git a/Engine/Plugins/Experimental/UAF/UAFAnimGraph/Source/UAFAnimGraph/Private/Traits/SequencePlayer.cpp b/Engine/Plugins/Experimental/UAF/UAFAnimGraph/Source/UAFAnimGraph/Private/Traits/SequencePlayer.cpp
--- a/Engine/Plugins/Experimental/UAF/UAFAnimGraph/Source/UAFAnimGraph/Private/Traits/SequencePlayer.cpp
+++ b/Engine/Plugins/Experimental/UAF/UAFAnimGraph/Source/UAFAnimGraph/Private/Traits/SequencePlayer.cpp
@@ -24,20 +24,6 @@ namespace UE::UAF
 	GENERATE_ANIM_TRAIT_IMPLEMENTATION(FSequencePlayerTrait, TRAIT_INTERFACE_ENUMERATOR, NULL_ANIM_TRAIT_INTERFACE_ENUMERATOR, NULL_ANIM_TRAIT_EVENT_ENUMERATOR)
 	#undef TRAIT_INTERFACE_ENUMERATOR
 
-	void FSequencePlayerTrait::FInstanceData::Construct(const FExecutionContext& Context, const FTraitBinding& Binding)
-	{
-		FTrait::FInstanceData::Construct(Context, Binding);
-
-		IGarbageCollection::RegisterWithGC(Context, Binding);
-	}
-
-	void FSequencePlayerTrait::FInstanceData::Destruct(const FExecutionContext& Context, const FTraitBinding& Binding)
-	{
-		FTrait::FInstanceData::Destruct(Context, Binding);
-
-		IGarbageCollection::UnregisterWithGC(Context, Binding);
-	}
-
 	void FSequencePlayerTrait::PreEvaluate(FEvaluateTraversalContext& Context, const TTraitBinding<IEvaluate>& Binding) const
 	{
 		const FSharedData* SharedData = Binding.GetSharedData<FSharedData>();
diff --git a/Engine/Plugins/Experimental/UAF/UAFAnimGraph/Source/UAFAnimGraph/Private/Traits/SequencePlayer.h b/Engine/Plugins/Experimental/UAF/UAFAnimGraph/Source/UAFAnimGraph/Private/Traits/SequencePlayer.h
--- a/Engine/Plugins/Experimental/UAF/UAFAnimGraph/Source/UAFAnimGraph/Private/Traits/SequencePlayer.h
+++ b/Engine/Plugins/Experimental/UAF/UAFAnimGraph/Source/UAFAnimGraph/Private/Traits/SequencePlayer.h
@@ -40,9 +40,6 @@ namespace UE::UAF
 
 			// The last advance type when FAnimationRuntime::AdvanceTime was called
 			ETypeAdvanceAnim LastAdvanceType = ETAA_Default;
-
-			void Construct(const FExecutionContext& Context, const FTraitBinding& Binding);
-			void Destruct(const FExecutionContext& Context, const FTraitBinding& Binding);
 		};
 
 		// IEvaluate impl
diff --git a/Engine/Plugins/Experimental/UAF/UAFAnimGraph/Source/UAFAnimGraph/Private/Traits/SubGraphHost.cpp b/Engine/Plugins/Experimental/UAF/UAFAnimGraph/Source/UAFAnimGraph/Private/Traits/SubGraphHost.cpp
--- a/Engine/Plugins/Experimental/UAF/UAFAnimGraph/Source/UAFAnimGraph/Private/Traits/SubGraphHost.cpp
+++ b/Engine/Plugins/Experimental/UAF/UAFAnimGraph/Source/UAFAnimGraph/Private/Traits/SubGraphHost.cpp
@@ -27,19 +27,10 @@ namespace UE::UAF
 	{
 		FTrait::FInstanceData::Construct(Context, Binding);
 
-		IGarbageCollection::RegisterWithGC(Context, Binding);
-
 		const FSharedData* SharedData = Binding.GetSharedData<FSharedData>();
 		ReferencePoseChildPtr = Context.AllocateNodeInstance(Binding, SharedData->ReferencePoseChild);
 	}
 
-	void FSubGraphHostTrait::FInstanceData::Destruct(const FExecutionContext& Context, const FTraitBinding& Binding)
-	{
-		FTrait::FInstanceData::Destruct(Context, Binding);
-
-		IGarbageCollection::UnregisterWithGC(Context, Binding);
-	}
-
 	uint32 FSubGraphHostTrait::GetNumChildren(const FExecutionContext& Context, const TTraitBinding<IHierarchy>& Binding) const
 	{
 		const FInstanceData* InstanceData = Binding.GetInstanceData<FInstanceData>();
diff --git a/Engine/Plugins/Experimental/UAF/UAFAnimGraph/Source/UAFAnimGraph/Private/Traits/SubGraphHost.h b/Engine/Plugins/Experimental/UAF/UAFAnimGraph/Source/UAFAnimGraph/Private/Traits/SubGraphHost.h
--- a/Engine/Plugins/Experimental/UAF/UAFAnimGraph/Source/UAFAnimGraph/Private/Traits/SubGraphHost.h
+++ b/Engine/Plugins/Experimental/UAF/UAFAnimGraph/Source/UAFAnimGraph/Private/Traits/SubGraphHost.h
@@ -82,7 +82,6 @@ namespace UE::UAF
 		struct FInstanceData : FTrait::FInstanceData
 		{
 			void Construct(const FExecutionContext& Context, const FTraitBinding& Binding);
-			void Destruct(const FExecutionContext& Context, const FTraitBinding& Binding);
 
 			// List of sub-graph slots
 			TArray<FSubGraphSlot> SubGraphSlots;
diff --git a/Engine/Plugins/Experimental/UAF/UAFAnimGraph/Source/UAFAnimGraph/Public/TraitCore/TraitTemplate.h b/Engine/Plugins/Experimental/UAF/UAFAnimGraph/Source/UAFAnimGraph/Public/TraitCore/TraitTemplate.h
--- a/Engine/Plugins/Experimental/UAF/UAFAnimGraph/Source/UAFAnimGraph/Public/TraitCore/TraitTemplate.h
+++ b/Engine/Plugins/Experimental/UAF/UAFAnimGraph/Source/UAFAnimGraph/Public/TraitCore/TraitTemplate.h
@@ -40,6 +40,9 @@ namespace UE::UAF
 		// Returns the trait mode
 		ETraitMode GetMode() const noexcept { return static_cast<ETraitMode>(Mode); }
 
+		// Returns whether this trait template needs be registered with the reference collector
+		bool IsGCAware() const noexcept { return bGCAware; }
+
 		// For base traits only, returns the number of traits on the sub-stack
 		uint32 GetNumStackTraits() const
 		{
@@ -140,10 +143,12 @@ namespace UE::UAF
 		friend struct FNodeTemplate;
 		friend struct FNodeTemplateBuilder;
 
-		FTraitTemplate(FTraitUID InUID, FTraitRegistryHandle InRegistryHandle, ETraitMode InMode, uint32 InTraitIndexOrNumTraits) noexcept
+		FTraitTemplate(FTraitUID InUID, FTraitRegistryHandle InRegistryHandle, ETraitMode InMode, bool bInGCAware, uint32 InTraitIndexOrNumTraits) noexcept
 			: UID(InUID.GetUID())
 			, RegistryHandle(InRegistryHandle)
 			, Mode(static_cast<uint8>(InMode))
+			, bGCAware(bInGCAware)
+			, Padding0(0)
 			, TraitIndexOrNumTraits(InTraitIndexOrNumTraits)
 			, NumLatentProperties(0)
 			, NumSubStackLatentProperties(0)
@@ -151,7 +156,7 @@ namespace UE::UAF
 			, NodeSharedOffset(0)
 			, NodeSharedLatentPropertyHandlesOffset(0)
 			, NodeInstanceOffset(0)
-			, Padding0(0)
+			, Padding1(0)
 		{
 		}
 
@@ -161,8 +166,15 @@ namespace UE::UAF
 		// Cached trait registry handle (16 bits)
 		FTraitRegistryHandle	RegistryHandle;
 
-		// Trait mode (we only need 1 bit, we could store other flags here)
-		uint8	Mode;
+		// Trait mode (we only need 2 bits)
+		uint8	Mode : 2;
+		static_assert(static_cast<uint8_t>(ETraitMode::Num) <= 2);
+		
+		// Flag set if a trait implements IGarbageCollection
+		uint8   bGCAware : 1;
+		
+		// (Padding for extra flags)
+		uint8   Padding0 : 5;
 
 		// For base traits, this contains the number of traits on our sub-stack
 		// For additive traits, this contains the trait index relative to the base trait
@@ -193,7 +205,7 @@ namespace UE::UAF
 		uint16	NodeSharedLatentPropertyHandlesOffset;	// Start of shared latent property handles for this trait
 		uint16	NodeInstanceOffset;						// Start of instance data for this trait
 
-		uint16	Padding0;						// Unused for now
+		uint16	Padding1;						// Unused for now
 
 		// TODO: We could cache which parent/base trait handles which interface
 		// This would avoid the need to iterate on every trait to look up a 'Super'. Perhaps only common interfaces could be cached.
diff --git a/Engine/Plugins/Experimental/UAF/UAFAnimGraph/Source/UAFAnimGraph/Public/TraitInterfaces/IGarbageCollection.h b/Engine/Plugins/Experimental/UAF/UAFAnimGraph/Source/UAFAnimGraph/Public/TraitInterfaces/IGarbageCollection.h
--- a/Engine/Plugins/Experimental/UAF/UAFAnimGraph/Source/UAFAnimGraph/Public/TraitInterfaces/IGarbageCollection.h
+++ b/Engine/Plugins/Experimental/UAF/UAFAnimGraph/Source/UAFAnimGraph/Public/TraitInterfaces/IGarbageCollection.h
@@ -21,14 +21,6 @@ namespace UE::UAF
 	{
 		DECLARE_ANIM_TRAIT_INTERFACE(IGarbageCollection)
 
-
-		// Registers the provided binding for GC callback
-		// Once registered, AddReferencedObjects is called during GC to collect references
-		static UE_API void RegisterWithGC(const FExecutionContext& Context, const FTraitBinding& Binding);
-
-		// Unregisters the provided binding from GC callback
-		static UE_API void UnregisterWithGC(const FExecutionContext& Context, const FTraitBinding& Binding);
-
 		// Called when garbage collection requests hard/strong object references
 		// @see UObject::AddReferencedObjects
 		UE_API virtual void AddReferencedObjects(const FExecutionContext& Context, const TTraitBinding<IGarbageCollection>& Binding, FReferenceCollector& Collector) const;
@@ -39,6 +31,16 @@ namespace UE::UAF
 		UE_API virtual const FText& GetDisplayName() const override;
 		UE_API virtual const FText& GetDisplayShortName() const override;
 #endif // WITH_EDITOR
+
+	private:
+		friend struct FExecutionContext;
+	
+		// Registers the provided binding for GC callback
+		// Once registered, AddReferencedObjects is called during GC to collect references
+		static UE_API void RegisterWithGC(const FExecutionContext& Context, const FTraitBinding& Binding);
+
+		// Unregisters the provided binding from GC callback
+		static UE_API void UnregisterWithGC(const FExecutionContext& Context, const FTraitBinding& Binding);
 	};
 
 	/**
diff --git a/Engine/Plugins/Experimental/UAF/UAFAnimGraph/Source/UAFAnimGraphUncookedOnly/Internal/Graph/AnimNextAnimationGraph_EditorData.h b/Engine/Plugins/Experimental/UAF/UAFAnimGraph/Source/UAFAnimGraphUncookedOnly/Internal/Graph/AnimNextAnimationGraph_EditorData.h
--- a/Engine/Plugins/Experimental/UAF/UAFAnimGraph/Source/UAFAnimGraphUncookedOnly/Internal/Graph/AnimNextAnimationGraph_EditorData.h
+++ b/Engine/Plugins/Experimental/UAF/UAFAnimGraph/Source/UAFAnimGraphUncookedOnly/Internal/Graph/AnimNextAnimationGraph_EditorData.h
@@ -59,6 +59,10 @@ class UAnimNextAnimationGraph_EditorData : public UAnimNextSharedVariables_Edito
 	friend class FAnimationAnimNextEditorTest_GraphTraitOperations;
 	friend class FAnimationAnimNextRuntimeTest_GraphExecute;
 	friend class FAnimationAnimNextRuntimeTest_GraphExecuteLatent;
+	friend class FAnimationUAFRuntimeTest_SharedDataGC;
+	friend class FAnimationUAFRuntimeTest_SharedDataMissingGC;
+	friend class FAnimationUAFRuntimeTest_InstanceDataGC;
+	friend class FAnimationUAFRuntimeTest_InstanceDataGCMissing;
 	friend class FAnimationAnimNextEditorTest_GraphManifest;
 
 public:
diff --git a/Engine/Plugins/Experimental/UAF/UAFAnimGraph/Source/UAFAnimGraphUncookedOnly/Private/Graph/AnimNextAnimationGraph_EditorData.cpp b/Engine/Plugins/Experimental/UAF/UAFAnimGraph/Source/UAFAnimGraphUncookedOnly/Private/Graph/AnimNextAnimationGraph_EditorData.cpp
--- a/Engine/Plugins/Experimental/UAF/UAFAnimGraph/Source/UAFAnimGraphUncookedOnly/Private/Graph/AnimNextAnimationGraph_EditorData.cpp
+++ b/Engine/Plugins/Experimental/UAF/UAFAnimGraph/Source/UAFAnimGraphUncookedOnly/Private/Graph/AnimNextAnimationGraph_EditorData.cpp
@@ -33,6 +33,7 @@
 #include "Graph/RigUnit_AnimNextGraphEvaluator.h"
 #include "Graph/RigUnit_AnimNextGraphRoot.h"
 #include "TraitCore/TraitInterfaceRegistry.h"
+#include "TraitInterfaces/IGarbageCollection.h"
 
 #include UE_INLINE_GENERATED_CPP_BY_NAME(AnimNextAnimationGraph_EditorData)
 
@@ -415,6 +416,42 @@ namespace UE::UAF::UncookedOnly::Private
 				}
 			}
 		}
+
+		const FTraitStackMapping& TraitStack = InMapping;
+		{
+			for (const FTraitEntryMapping& TraitEntry : TraitStack.TraitEntries)
+			{
+				for (URigVMPin* Pin : TraitEntry.DecoratorEntryPin->GetSubPins())
+				{
+					if (!Pin->IsLazy())
+					{
+						continue;
+					}
+					
+					const bool bIsProgrammaticPin = Pin->IsProgrammaticPin();
+					const bool bHasLinks = !Pin->GetLinks().IsEmpty();
+					if (bHasLinks || bIsProgrammaticPin)
+					{
+						// This pin has something linked to it, it is a latent pin
+						FString PinCPPType = Pin->GetCPPType();
+						if (RigVMTypeUtils::ObjectFromCPPType(PinCPPType))
+						{
+							const bool bImplementsGarbageCollection = TraitEntry.Trait->GetTraitInterface(IGarbageCollection::InterfaceUID) != nullptr;
+							if (!bImplementsGarbageCollection)
+							{
+								const FText MissingInterfaceError = FText::Format(
+									LOCTEXT("TraitStatusInStack_DoesNotImplementGarbageCollection", "@@: Trait is connected to a UObject but does not implement IGarbageCollection ({0})"),
+									TraitEntry.Trait->GetTraitSharedDataStruct()->GetDisplayNameText());
+								InSettings.ASTSettings.Report(EMessageSeverity::Error, InVMUnitNode, MissingInterfaceError.ToString());
+								bResult = false;
+								break;
+							}
+						}
+					}
+				}
+			}
+		}
+		
 		return bResult;
 	}
 
diff --git a/Engine/Plugins/Experimental/UAF/UAFControlRig/Source/UAFControlRig/Private/ControlRigTrait.cpp b/Engine/Plugins/Experimental/UAF/UAFControlRig/Source/UAFControlRig/Private/ControlRigTrait.cpp
--- a/Engine/Plugins/Experimental/UAF/UAFControlRig/Source/UAFControlRig/Private/ControlRigTrait.cpp
+++ b/Engine/Plugins/Experimental/UAF/UAFControlRig/Source/UAFControlRig/Private/ControlRigTrait.cpp
@@ -198,8 +198,6 @@ void FControlRigTrait::FInstanceData::Construct(const FExecutionContext& Context
 
 	const FSharedData* SharedData = Binding.GetSharedData<FSharedData>();
 
-	IGarbageCollection::RegisterWithGC(Context, Binding);
-
 #if WITH_EDITOR
 	{
 		OnObjectsReinstancedHandle = UE::UAF::ControlRig::FAnimNextControlRigModule::OnObjectsReinstanced.AddRaw(this, &FControlRigTrait::FInstanceData::OnObjectsReinstanced);
@@ -227,8 +225,6 @@ void FControlRigTrait::FInstanceData::Destruct(const FExecutionContext& Context,
 {
 	FTrait::FInstanceData::Destruct(Context, Binding);
 
-	IGarbageCollection::UnregisterWithGC(Context, Binding);
-
 #if WITH_EDITOR
 	UE::UAF::ControlRig::FAnimNextControlRigModule::OnObjectsReinstanced.Remove(OnObjectsReinstancedHandle);
 #endif
diff --git a/Engine/Plugins/Experimental/UAF/UAFMirroring/Source/UAFMirroring/Private/MirroringTrait.cpp b/Engine/Plugins/Experimental/UAF/UAFMirroring/Source/UAFMirroring/Private/MirroringTrait.cpp
--- a/Engine/Plugins/Experimental/UAF/UAFMirroring/Source/UAFMirroring/Private/MirroringTrait.cpp
+++ b/Engine/Plugins/Experimental/UAF/UAFMirroring/Source/UAFMirroring/Private/MirroringTrait.cpp
@@ -118,15 +118,6 @@ namespace UE::UAF
 		{
 			InstanceData->Input = Context.AllocateNodeInstance(Binding, SharedData->Input);
 		}
-		
-		IGarbageCollection::RegisterWithGC(Context, Binding);
-	}
-
-	void FMirroringTrait::FInstanceData::Deconstruct(const FExecutionContext& Context, const FTraitBinding& Binding)
-	{
-		FTrait::FInstanceData::Destruct(Context, Binding);
-
-		IGarbageCollection::UnregisterWithGC(Context, Binding);
 	}
 
 	void FMirroringTrait::PostEvaluate(FEvaluateTraversalContext& Context, const TTraitBinding<IEvaluate>& Binding) const
@@ -186,20 +177,6 @@ namespace UE::UAF
 	GENERATE_ANIM_TRAIT_IMPLEMENTATION(FMirroringAdditiveTrait, TRAIT_INTERFACE_ENUMERATOR, NULL_ANIM_TRAIT_INTERFACE_ENUMERATOR, NULL_ANIM_TRAIT_EVENT_ENUMERATOR)
 	#undef TRAIT_INTERFACE_ENUMERATOR
 
-	void FMirroringAdditiveTrait::FInstanceData::Construct(const FExecutionContext& Context, const FTraitBinding& Binding)
-	{
-		FTrait::FInstanceData::Construct(Context, Binding);
-
-		IGarbageCollection::RegisterWithGC(Context, Binding);
-	}
-
-	void FMirroringAdditiveTrait::FInstanceData::Deconstruct(const FExecutionContext& Context, const FTraitBinding& Binding)
-	{
-		FTrait::FInstanceData::Destruct(Context, Binding);
-
-		IGarbageCollection::UnregisterWithGC(Context, Binding);
-	}
-
 	void FMirroringAdditiveTrait::PostEvaluate(FEvaluateTraversalContext& Context, const TTraitBinding<IEvaluate>& Binding) const
 	{
 		IEvaluate::PostEvaluate(Context, Binding);
diff --git a/Engine/Plugins/Experimental/UAF/UAFMirroring/Source/UAFMirroring/Private/MirroringTrait.h b/Engine/Plugins/Experimental/UAF/UAFMirroring/Source/UAFMirroring/Private/MirroringTrait.h
--- a/Engine/Plugins/Experimental/UAF/UAFMirroring/Source/UAFMirroring/Private/MirroringTrait.h
+++ b/Engine/Plugins/Experimental/UAF/UAFMirroring/Source/UAFMirroring/Private/MirroringTrait.h
@@ -99,7 +99,6 @@ namespace UE::UAF
 			FMirroringTraitCache Cache;
 			
 			void Construct(const FExecutionContext& Context, const FTraitBinding& Binding);
-			void Deconstruct(const FExecutionContext& Context, const FTraitBinding& Binding);
 		};
 		
 		// IEvaluate impl
@@ -137,9 +136,6 @@ namespace UE::UAF
 			
 			// Data that stays the same between evaluation passes.
 			FMirroringTraitCache Cache;
-
-			void Construct(const FExecutionContext& Context, const FTraitBinding& Binding);
-			void Deconstruct(const FExecutionContext& Context, const FTraitBinding& Binding);
 		};
 
 		// IEvaluate impl
diff --git a/Engine/Plugins/Experimental/UAF/UAFPoseSearch/Source/UAFPoseSearch/Private/MotionMatchingTrait.cpp b/Engine/Plugins/Experimental/UAF/UAFPoseSearch/Source/UAFPoseSearch/Private/MotionMatchingTrait.cpp
--- a/Engine/Plugins/Experimental/UAF/UAFPoseSearch/Source/UAFPoseSearch/Private/MotionMatchingTrait.cpp
+++ b/Engine/Plugins/Experimental/UAF/UAFPoseSearch/Source/UAFPoseSearch/Private/MotionMatchingTrait.cpp
@@ -62,19 +62,7 @@ PRAGMA_DISABLE_DEPRECATION_WARNINGS
 	GENERATE_ANIM_TRAIT_IMPLEMENTATION(FMotionMatchingTrait, TRAIT_INTERFACE_ENUMERATOR, NULL_ANIM_TRAIT_INTERFACE_ENUMERATOR, NULL_ANIM_TRAIT_EVENT_ENUMERATOR)
 #undef TRAIT_INTERFACE_ENUMERATOR
 PRAGMA_ENABLE_DEPRECATION_WARNINGS
-
-	void FMotionMatchingTrait::FInstanceData::Construct(const FExecutionContext& Context, const FTraitBinding& Binding)
-	{
-		FTrait::FInstanceData::Construct(Context, Binding);
-		IGarbageCollection::RegisterWithGC(Context, Binding);
-	}
-
-	void FMotionMatchingTrait::FInstanceData::Destruct(const FExecutionContext& Context, const FTraitBinding& Binding)
-	{
-		FTrait::FInstanceData::Destruct(Context, Binding);
-		IGarbageCollection::UnregisterWithGC(Context, Binding);
-	}
-
+	
 	void FMotionMatchingTrait::OnBecomeRelevant(FUpdateTraversalContext& Context, const TTraitBinding<IUpdate>& Binding, const FTraitUpdateState& TraitState) const
 	{
 		IUpdate::OnBecomeRelevant(Context, Binding, TraitState);
diff --git a/Engine/Plugins/Experimental/UAF/UAFPoseSearch/Source/UAFPoseSearch/Private/MotionMatchingTrait.h b/Engine/Plugins/Experimental/UAF/UAFPoseSearch/Source/UAFPoseSearch/Private/MotionMatchingTrait.h
--- a/Engine/Plugins/Experimental/UAF/UAFPoseSearch/Source/UAFPoseSearch/Private/MotionMatchingTrait.h
+++ b/Engine/Plugins/Experimental/UAF/UAFPoseSearch/Source/UAFPoseSearch/Private/MotionMatchingTrait.h
@@ -32,9 +32,6 @@ struct FMotionMatchingTrait : FAdditiveTrait, IUpdate, IEvaluate, IGarbageCollec
 #if WITH_EDITOR
 		bool bIsPostEvaluateBeingCalled = true;
 #endif // WITH_EDITOR
-
-		void Construct(const FExecutionContext& Context, const FTraitBinding& Binding);
-		void Destruct(const FExecutionContext& Context, const FTraitBinding& Binding);
 	};
 
 	// IUpdate impl
diff --git a/Engine/Plugins/Experimental/UAF/UAFStateTree/Source/UAFStateTree/Internal/AnimStateTreeTrait.h b/Engine/Plugins/Experimental/UAF/UAFStateTree/Source/UAFStateTree/Internal/AnimStateTreeTrait.h
--- a/Engine/Plugins/Experimental/UAF/UAFStateTree/Source/UAFStateTree/Internal/AnimStateTreeTrait.h
+++ b/Engine/Plugins/Experimental/UAF/UAFStateTree/Source/UAFStateTree/Internal/AnimStateTreeTrait.h
@@ -49,7 +49,6 @@ struct FStateTreeTrait : FAdditiveTrait, IUpdate, IGarbageCollection
 		FStateTreeInstanceData InstanceData;
 		FStateTreeExecutionContext::FExternalGlobalParameters StateTreeExternalParameters;
 
-		void Construct(const FExecutionContext& Context, const FTraitBinding& Binding);
 		void Destruct(const FExecutionContext& Context, const FTraitBinding& Binding);
 	};
 
diff --git a/Engine/Plugins/Experimental/UAF/UAFStateTree/Source/UAFStateTree/Private/AnimStateTreeTrait.cpp b/Engine/Plugins/Experimental/UAF/UAFStateTree/Source/UAFStateTree/Private/AnimStateTreeTrait.cpp
--- a/Engine/Plugins/Experimental/UAF/UAFStateTree/Source/UAFStateTree/Private/AnimStateTreeTrait.cpp
+++ b/Engine/Plugins/Experimental/UAF/UAFStateTree/Source/UAFStateTree/Private/AnimStateTreeTrait.cpp
@@ -29,13 +29,7 @@ namespace UE::UAF
 #if ENABLE_ANIM_DEBUG 
 	TAutoConsoleVariable<bool> CVarLogPropertyBindingMemoryPtrInfo(TEXT("a.StateTree.LogPropertyBindingMemoryPtrInfo"), false, TEXT("Log information while generating property binding memory pointer mappings between RigVM and StateTree"));
 #endif //ENABLE_ANIM_DEBUG 
-
-	void FStateTreeTrait::FInstanceData::Construct(const FExecutionContext& Context, const FTraitBinding& Binding)
-	{
-		FTrait::FInstanceData::Construct(Context, Binding);
-		IGarbageCollection::RegisterWithGC(Context, Binding);
-	}
-
+	
 	void FStateTreeTrait::FInstanceData::Destruct(const FExecutionContext& Context, const FTraitBinding& Binding)
 	{
 #if WITH_STATETREE_DEBUG
@@ -75,7 +69,6 @@ namespace UE::UAF
 		}*/	
 #endif
 		FTrait::FInstanceData::Destruct(Context, Binding);
-		IGarbageCollection::UnregisterWithGC(Context, Binding);
 	}
 
 	void FStateTreeTrait::OnBecomeRelevant(FUpdateTraversalContext& Context, const TTraitBinding<IUpdate>& Binding, const FTraitUpdateState& TraitState) const
diff --git a/Engine/Plugins/Experimental/UAF/UAFTests/Source/UAFAnimGraphTestSuite/Private/GCInstanceDataGCTest.cpp b/Engine/Plugins/Experimental/UAF/UAFTests/Source/UAFAnimGraphTestSuite/Private/GCInstanceDataGCTest.cpp
new file mode 100644
--- /dev/null
+++ b/Engine/Plugins/Experimental/UAF/UAFTests/Source/UAFAnimGraphTestSuite/Private/GCInstanceDataGCTest.cpp
@@ -0,0 +1,204 @@
+// Copyright Epic Games, Inc. All Rights Reserved.
+
+#include "GCTestsUtil.h"
+
+#include "AnimNextRuntimeTest.h"
+#include "AnimNextTest.h"
+#include "UncookedOnlyUtils.h"
+#include "Graph/AnimNextAnimationGraph.h"
+#include "Graph/AnimNextAnimationGraphFactory.h"
+#include "Graph/AnimNextAnimationGraph_EditorData.h"
+#include "Graph/AnimNextGraphInstance.h"
+#include "Graph/RigDecorator_AnimNextCppTrait.h"
+#include "Graph/RigUnit_AnimNextGraphRoot.h"
+#include "Graph/RigUnit_AnimNextTraitStack.h"
+#include "Misc/AutomationTest.h"
+#include "Script/UAFRigVMComponent.h"
+#include "TraitCore/Trait.h"
+#include "TraitCore/TraitRegistry.h"
+#include "TraitInterfaces/IEvaluate.h"
+#include "TraitInterfaces/IGarbageCollection.h"
+#include "TraitInterfaces/IUpdate.h"
+
+#if WITH_DEV_AUTOMATION_TESTS
+
+namespace UE::UAF
+{
+	/** A test trait with a UObject instance property that correctly implements IGarbageCollection */
+	struct FTestTrait_InstanceDataGC : FBaseTrait, IUpdate, IGarbageCollection
+	{
+		DECLARE_ANIM_TRAIT(FTestTrait_InstanceDataGC, FBaseTrait)
+
+		struct FInstanceData : FTrait::FInstanceData
+		{
+			TObjectPtr<UAnimSequence> AnimSequence;
+			TWeakObjectPtr<UAnimSequence> WeakAnimSequence;
+			uint8 UpdateCount = 0;
+		};
+
+		// IUpdate impl
+		virtual void PostUpdate(FUpdateTraversalContext& Context, const TTraitBinding<IUpdate>& Binding, const FTraitUpdateState& TraitState) const override
+		{
+			IUpdate::PostUpdate(Context, Binding, TraitState);
+
+			FInstanceData* InstanceData = Binding.GetInstanceData<FInstanceData>();
+			const FSharedData* SharedData = Binding.GetSharedData<FSharedData>();
+
+			if (InstanceData->UpdateCount == 0)
+			{
+				InstanceData->AnimSequence = NewObject<UAnimSequence>(GetTransientPackage(), NAME_None);
+				InstanceData->WeakAnimSequence = InstanceData->AnimSequence;
+			}
+			
+			const bool bIsValid = InstanceData->WeakAnimSequence.IsValid() && !InstanceData->WeakAnimSequence->HasAnyFlags(RF_BeginDestroyed | RF_FinishDestroyed);
+			
+			FRigVMExecuteContext& ExecuteContext = Context.GetRootGraphInstance().GetComponent<FUAFRigVMComponent>().GetExtendedExecuteContext().GetPublicData();
+			ExecuteContext.Logf(EMessageSeverity::Info, TEXT("%s"), bIsValid ? TEXT("Valid") : TEXT("Invalid"));
+
+			++InstanceData->UpdateCount;
+		}
+
+		// IGarbageCollection impl
+		virtual void AddReferencedObjects(const FExecutionContext& Context, const TTraitBinding<IGarbageCollection>& Binding, FReferenceCollector& Collector) const override
+		{
+			IGarbageCollection::AddReferencedObjects(Context, Binding, Collector);
+
+			FInstanceData* InstanceData = Binding.GetInstanceData<FInstanceData>();
+			const FSharedData* SharedData = Binding.GetSharedData<FSharedData>();
+
+			Collector.AddReferencedObject(InstanceData->AnimSequence);
+		}
+	};
+
+	// Trait implementation boilerplate
+	#define TRAIT_INTERFACE_ENUMERATOR(GeneratorMacro) \
+		GeneratorMacro(IUpdate) \
+		GeneratorMacro(IGarbageCollection) \
+	
+	GENERATE_ANIM_TRAIT_IMPLEMENTATION(FTestTrait_InstanceDataGC, TRAIT_INTERFACE_ENUMERATOR, NULL_ANIM_TRAIT_INTERFACE_ENUMERATOR, NULL_ANIM_TRAIT_EVENT_ENUMERATOR)
+	#undef TRAIT_INTERFACE_ENUMERATOR
+}
+
+IMPLEMENT_SIMPLE_AUTOMATION_TEST(FAnimationUAFRuntimeTest_InstanceDataGC, "Animation.UAF.Runtime.GCInstanceDataWithGCInterface", EAutomationTestFlags::EditorContext | EAutomationTestFlags::EngineFilter)
+
+/*
+ * Tests that a trait that has a UObject instance property that incorrectly does not implement IGarbageCollection results in its UObject being GC'd.
+ */
+bool FAnimationUAFRuntimeTest_InstanceDataGC::RunTest(const FString& InParameters)
+{
+	using namespace UE::UAF;
+
+	{
+		AUTO_REGISTER_ANIM_TRAIT(FTestTrait_InstanceDataGC)
+		FScopedClearNodeTemplateRegistry ScopedClearNodeTemplateRegistry;
+
+		UFactory* GraphFactory = NewObject<UAnimNextAnimationGraphFactory>();
+		UAnimNextAnimationGraph* AnimationGraph = CastChecked<UAnimNextAnimationGraph>(GraphFactory->FactoryCreateNew(UAnimNextAnimationGraph::StaticClass(), GetTransientPackage(), TEXT("TestAnimNextGraph"), RF_Transient, nullptr, nullptr, NAME_None));
+		UE_RETURN_ON_ERROR(AnimationGraph != nullptr, "FAnimationUAFRuntimeTest_InstanceDataGC -> Failed to create animation graph");
+
+		UAnimNextAnimationGraph_EditorData* EditorData = UncookedOnly::FUtils::GetEditorData<UAnimNextAnimationGraph_EditorData>(AnimationGraph);
+		UE_RETURN_ON_ERROR(EditorData != nullptr, "FAnimationUAFRuntimeTest_InstanceDataGC -> Failed to find module editor data");
+
+		UAnimNextController* Controller = Cast<UAnimNextController>(EditorData->GetRigVMClient()->GetController(EditorData->GetRigVMClient()->GetDefaultModel()));
+		UE_RETURN_ON_ERROR(Controller != nullptr, "FAnimationUAFRuntimeTest_InstanceDataGC -> Failed to get RigVM controller");
+
+		// Find graph entry point
+		URigVMNode* MainEntryPointNode = Controller->GetGraph()->FindNodeByName(FRigUnit_AnimNextGraphRoot::StaticStruct()->GetFName());
+		UE_RETURN_ON_ERROR(MainEntryPointNode != nullptr, "FAnimationUAFRuntimeTest_InstanceDataGC -> Failed to find main entry point node");
+
+		URigVMPin* BeginExecutePin = MainEntryPointNode->FindPin(GET_MEMBER_NAME_STRING_CHECKED(FRigUnit_AnimNextGraphRoot, Result));
+		UE_RETURN_ON_ERROR(BeginExecutePin != nullptr && BeginExecutePin->GetDirection() == ERigVMPinDirection::Input, "FAnimationUAFRuntimeTest_InstanceDataGC -> Failed to create entry point");
+
+		URigVMUnitNode* DecoratorStackNode = nullptr;
+		FString DisplayName;
+		{
+			// Suspend auto compilation until we have constructed a valid trait stack
+			TGuardValue<bool> SuspendCompile(EditorData->bAutoRecompileVM, false);
+
+			// Create an empty trait stack node
+			DecoratorStackNode = Controller->AddUnitNode(FRigUnit_AnimNextTraitStack::StaticStruct(), FRigVMStruct::ExecuteName, FVector2D(0.0f, 0.0f), FString(), false);
+			UE_RETURN_ON_ERROR(DecoratorStackNode != nullptr, "FAnimationUAFRuntimeTest_InstanceDataGC -> Failed to create trait stack node");
+
+			// Link our stack result to our entry point
+			Controller->AddLink(DecoratorStackNode->GetPins()[0], MainEntryPointNode->GetPins()[0]);
+
+			// Add a trait
+			const UScriptStruct* CppDecoratorStruct = FRigDecorator_AnimNextCppDecorator::StaticStruct();
+			UE_RETURN_ON_ERROR(CppDecoratorStruct != nullptr, "FAnimationUAFRuntimeTest_InstanceDataGC -> Failed to get find Cpp trait static struct");
+
+			const FTrait* Trait = FTraitRegistry::Get().Find(FTestTrait_InstanceDataGC::TraitUID);
+			UE_RETURN_ON_ERROR(Trait != nullptr, "FAnimationUAFRuntimeTest_InstanceDataGC -> Failed to find test trait");
+
+			UScriptStruct* ScriptStruct = Trait->GetTraitSharedDataStruct();
+			UE_RETURN_ON_ERROR(ScriptStruct != nullptr, "FAnimationUAFRuntimeTest_InstanceDataGC -> Failed to find trait shared data struct");
+
+			FString DefaultValue;
+			{
+				const FRigDecorator_AnimNextCppDecorator DefaultCppDecoratorStructInstance;
+				FRigDecorator_AnimNextCppDecorator CppDecoratorStructInstance;
+				CppDecoratorStructInstance.DecoratorSharedDataStruct = ScriptStruct;
+
+				UE_RETURN_ON_ERROR(CppDecoratorStructInstance.CanBeAddedToNode(DecoratorStackNode, nullptr), "FAnimationUAFRuntimeTest_InstanceDataGC -> Trait cannot be added to trait stack node");
+
+				FRigDecorator_AnimNextCppDecorator::StaticStruct()->ExportText(DefaultValue, &CppDecoratorStructInstance, &DefaultCppDecoratorStructInstance, nullptr, PPF_None, nullptr);
+			}
+
+			FString DisplayNameMetadata;
+			ScriptStruct->GetStringMetaDataHierarchical(FRigVMStruct::DisplayNameMetaName, &DisplayNameMetadata);
+			DisplayName = DisplayNameMetadata.IsEmpty() ? Trait->GetTraitName() : DisplayNameMetadata;
+
+			const FName DecoratorName = Controller->AddTrait(
+				DecoratorStackNode->GetFName(),
+				*CppDecoratorStruct->GetPathName(),
+				*DisplayName,
+				DefaultValue, INDEX_NONE, true, true);
+			UE_RETURN_ON_ERROR(DecoratorName == DisplayName, TEXT("FAnimationUAFRuntimeTest_InstanceDataGC -> Unexpected trait name"));
+		}
+
+		// Set some values on our trait
+		URigVMPin* DecoratorPin = DecoratorStackNode->FindPin(*DisplayName);
+		UE_RETURN_ON_ERROR(DecoratorPin != nullptr, TEXT("FAnimationUAFRuntimeTest_InstanceDataGC -> Failed to find trait pin"));
+
+		// For some reason a manual recompile is needed, otherwise the trait interface implementations are never run
+		EditorData->RecompileVM();
+
+		TSharedPtr<FAnimNextGraphInstance> GraphInstance = AnimationGraph->AllocateInstance();
+
+		// We place the graph instance in a UObject so that its contained within something the reference collector is aware of
+		// so that FGCGraphInstanceComponent.AddStructReferencedObjects is called. (This isn't strictly needed as it won't traverse
+		// our test trait anyway as it hasn't implemented IGarbageCollection) but its to match parity with the other GC test
+		TObjectPtr<UGraphInstanceHolder> GraphInstanceHolder = NewObject<UGraphInstanceHolder>();
+		GraphInstanceHolder->AddToRoot();
+		GraphInstanceHolder->GraphInstance = GraphInstance;
+		
+		TArray<FString> Messages;
+		FRigVMRuntimeSettings RuntimeSettings;
+		RuntimeSettings.SetLogFunction([&Messages](const FRigVMLogSettings& InLogSettings, const FRigVMExecuteContext* InContext, const FString& Message)
+		{
+			Messages.Add(Message);
+		});
+		GraphInstance->GetComponent<FUAFRigVMComponent>().GetExtendedExecuteContext().SetRuntimeSettings(RuntimeSettings);
+		
+		FUpdateGraphContext UpdateGraphContext(*GraphInstance.Get(), 1.0f / 30.0f);
+		// First update, the anim sequence has been added to root so should definitely be alive
+		UE::UAF::UpdateGraph(UpdateGraphContext);
+
+		CollectGarbage(GARBAGE_COLLECTION_KEEPFLAGS);
+
+		// Second update, the anim sequence should have been destroyed now
+		UE::UAF::UpdateGraph(UpdateGraphContext);
+
+		UE_RETURN_ON_ERROR(Messages.Num() == 2, "FAnimationUAFRuntimeTest_InstanceDataGC -> Not enough messages received");
+		AddErrorIfFalse(Messages[0] == TEXT("Valid"), "FAnimationUAFRuntimeTest_InstanceDataGC -> Update 1: Test AnimSequence should be valid");
+		AddErrorIfFalse(Messages[1] == TEXT("Valid"), "FAnimationUAFRuntimeTest_InstanceDataGC -> Update 2: Test AnimSequence should be valid");
+
+		GraphInstanceHolder->GraphInstance = nullptr;
+		GraphInstanceHolder->RemoveFromRoot();
+	}
+
+	Tests::FUtils::CleanupAfterTests();
+
+	return true;
+}
+
+#endif // WITH_DEV_AUTOMATION_TESTS
\ No newline at end of file
diff --git a/Engine/Plugins/Experimental/UAF/UAFTests/Source/UAFAnimGraphTestSuite/Private/GCInstanceDataMissingGCTest.cpp b/Engine/Plugins/Experimental/UAF/UAFTests/Source/UAFAnimGraphTestSuite/Private/GCInstanceDataMissingGCTest.cpp
new file mode 100644
--- /dev/null
+++ b/Engine/Plugins/Experimental/UAF/UAFTests/Source/UAFAnimGraphTestSuite/Private/GCInstanceDataMissingGCTest.cpp
@@ -0,0 +1,189 @@
+// Copyright Epic Games, Inc. All Rights Reserved.
+
+#include "GCTestsUtil.h"
+
+#include "AnimNextRuntimeTest.h"
+#include "AnimNextTest.h"
+#include "UncookedOnlyUtils.h"
+#include "Graph/AnimNextAnimationGraph.h"
+#include "Graph/AnimNextAnimationGraphFactory.h"
+#include "Graph/AnimNextAnimationGraph_EditorData.h"
+#include "Graph/AnimNextGraphInstance.h"
+#include "Graph/RigDecorator_AnimNextCppTrait.h"
+#include "Graph/RigUnit_AnimNextGraphRoot.h"
+#include "Graph/RigUnit_AnimNextTraitStack.h"
+#include "Misc/AutomationTest.h"
+#include "Script/UAFRigVMComponent.h"
+#include "TraitCore/Trait.h"
+#include "TraitCore/TraitRegistry.h"
+#include "TraitInterfaces/IEvaluate.h"
+#include "TraitInterfaces/IUpdate.h"
+
+#if WITH_DEV_AUTOMATION_TESTS
+
+namespace UE::UAF
+{
+	/** A test trait with a UObject instance property that incorrectly does not implement IGarbageCollection */
+	struct FTestTrait_InstanceDataGCMissing : FBaseTrait, IUpdate
+	{
+		DECLARE_ANIM_TRAIT(FTestTrait_InstanceDataGCMissing, FBaseTrait)
+
+		struct FInstanceData : FTrait::FInstanceData
+		{
+			TWeakObjectPtr<UAnimSequence> WeakAnimSequence;
+			uint8 UpdateCount = 0;
+		};
+
+		// IUpdate impl
+		virtual void PostUpdate(FUpdateTraversalContext& Context, const TTraitBinding<IUpdate>& Binding, const FTraitUpdateState& TraitState) const override
+		{
+			IUpdate::PostUpdate(Context, Binding, TraitState);
+
+			FInstanceData* InstanceData = Binding.GetInstanceData<FInstanceData>();
+			const FSharedData* SharedData = Binding.GetSharedData<FSharedData>();
+
+			if (InstanceData->UpdateCount == 0)
+			{
+				InstanceData->WeakAnimSequence = NewObject<UAnimSequence>(GetTransientPackage(), NAME_None);
+			}
+			
+			const bool bIsValid = InstanceData->WeakAnimSequence.IsValid() && !InstanceData->WeakAnimSequence->HasAnyFlags(RF_BeginDestroyed | RF_FinishDestroyed);
+			
+			FRigVMExecuteContext& ExecuteContext = Context.GetRootGraphInstance().GetComponent<FUAFRigVMComponent>().GetExtendedExecuteContext().GetPublicData();
+			ExecuteContext.Logf(EMessageSeverity::Info, TEXT("%s"), bIsValid ? TEXT("Valid") : TEXT("Invalid"));
+
+			++InstanceData->UpdateCount;
+		}
+	};
+
+	// Trait implementation boilerplate
+	#define TRAIT_INTERFACE_ENUMERATOR(GeneratorMacro) \
+		GeneratorMacro(IUpdate) \
+	
+	GENERATE_ANIM_TRAIT_IMPLEMENTATION(FTestTrait_InstanceDataGCMissing, TRAIT_INTERFACE_ENUMERATOR, NULL_ANIM_TRAIT_INTERFACE_ENUMERATOR, NULL_ANIM_TRAIT_EVENT_ENUMERATOR)
+	#undef TRAIT_INTERFACE_ENUMERATOR
+}
+
+IMPLEMENT_SIMPLE_AUTOMATION_TEST(FAnimationUAFRuntimeTest_InstanceDataGCMissing, "Animation.UAF.Runtime.GCInstanceDataMissingGCInterface", EAutomationTestFlags::EditorContext | EAutomationTestFlags::EngineFilter)
+
+/*
+ * Tests that a trait that has a UObject instance property that implements IGarbageCollection does not have its UObject GC'd.
+ */
+bool FAnimationUAFRuntimeTest_InstanceDataGCMissing::RunTest(const FString& InParameters)
+{
+	using namespace UE::UAF;
+
+	{
+		AUTO_REGISTER_ANIM_TRAIT(FTestTrait_InstanceDataGCMissing)
+		FScopedClearNodeTemplateRegistry ScopedClearNodeTemplateRegistry;
+
+		UFactory* GraphFactory = NewObject<UAnimNextAnimationGraphFactory>();
+		UAnimNextAnimationGraph* AnimationGraph = CastChecked<UAnimNextAnimationGraph>(GraphFactory->FactoryCreateNew(UAnimNextAnimationGraph::StaticClass(), GetTransientPackage(), TEXT("TestAnimNextGraph"), RF_Transient, nullptr, nullptr, NAME_None));
+		UE_RETURN_ON_ERROR(AnimationGraph != nullptr, "FAnimationUAFRuntimeTest_InstanceDataGCMissing -> Failed to create animation graph");
+
+		UAnimNextAnimationGraph_EditorData* EditorData = UncookedOnly::FUtils::GetEditorData<UAnimNextAnimationGraph_EditorData>(AnimationGraph);
+		UE_RETURN_ON_ERROR(EditorData != nullptr, "FAnimationUAFRuntimeTest_InstanceDataGCMissing -> Failed to find module editor data");
+
+		UAnimNextController* Controller = Cast<UAnimNextController>(EditorData->GetRigVMClient()->GetController(EditorData->GetRigVMClient()->GetDefaultModel()));
+		UE_RETURN_ON_ERROR(Controller != nullptr, "FAnimationUAFRuntimeTest_InstanceDataGCMissing -> Failed to get RigVM controller");
+
+		// Find graph entry point
+		URigVMNode* MainEntryPointNode = Controller->GetGraph()->FindNodeByName(FRigUnit_AnimNextGraphRoot::StaticStruct()->GetFName());
+		UE_RETURN_ON_ERROR(MainEntryPointNode != nullptr, "FAnimationUAFRuntimeTest_InstanceDataGCMissing -> Failed to find main entry point node");
+
+		URigVMPin* BeginExecutePin = MainEntryPointNode->FindPin(GET_MEMBER_NAME_STRING_CHECKED(FRigUnit_AnimNextGraphRoot, Result));
+		UE_RETURN_ON_ERROR(BeginExecutePin != nullptr && BeginExecutePin->GetDirection() == ERigVMPinDirection::Input, "FAnimationUAFRuntimeTest_InstanceDataGCMissing -> Failed to create entry point");
+
+		URigVMUnitNode* DecoratorStackNode = nullptr;
+		FString DisplayName;
+		{
+			// Suspend auto compilation until we have constructed a valid trait stack
+			TGuardValue<bool> SuspendCompile(EditorData->bAutoRecompileVM, false);
+
+			// Create an empty trait stack node
+			DecoratorStackNode = Controller->AddUnitNode(FRigUnit_AnimNextTraitStack::StaticStruct(), FRigVMStruct::ExecuteName, FVector2D(0.0f, 0.0f), FString(), false);
+			UE_RETURN_ON_ERROR(DecoratorStackNode != nullptr, "FAnimationUAFRuntimeTest_InstanceDataGCMissing -> Failed to create trait stack node");
+
+			// Link our stack result to our entry point
+			Controller->AddLink(DecoratorStackNode->GetPins()[0], MainEntryPointNode->GetPins()[0]);
+
+			// Add a trait
+			const UScriptStruct* CppDecoratorStruct = FRigDecorator_AnimNextCppDecorator::StaticStruct();
+			UE_RETURN_ON_ERROR(CppDecoratorStruct != nullptr, "FAnimationUAFRuntimeTest_InstanceDataGCMissing -> Failed to get find Cpp trait static struct");
+
+			const FTrait* Trait = FTraitRegistry::Get().Find(FTestTrait_InstanceDataGCMissing::TraitUID);
+			UE_RETURN_ON_ERROR(Trait != nullptr, "FAnimationUAFRuntimeTest_InstanceDataGCMissing -> Failed to find test trait");
+
+			UScriptStruct* ScriptStruct = Trait->GetTraitSharedDataStruct();
+			UE_RETURN_ON_ERROR(ScriptStruct != nullptr, "FAnimationUAFRuntimeTest_InstanceDataGCMissing -> Failed to find trait shared data struct");
+
+			FString DefaultValue;
+			{
+				const FRigDecorator_AnimNextCppDecorator DefaultCppDecoratorStructInstance;
+				FRigDecorator_AnimNextCppDecorator CppDecoratorStructInstance;
+				CppDecoratorStructInstance.DecoratorSharedDataStruct = ScriptStruct;
+
+				UE_RETURN_ON_ERROR(CppDecoratorStructInstance.CanBeAddedToNode(DecoratorStackNode, nullptr), "FAnimationUAFRuntimeTest_InstanceDataGCMissing -> Trait cannot be added to trait stack node");
+
+				FRigDecorator_AnimNextCppDecorator::StaticStruct()->ExportText(DefaultValue, &CppDecoratorStructInstance, &DefaultCppDecoratorStructInstance, nullptr, PPF_None, nullptr);
+			}
+
+			FString DisplayNameMetadata;
+			ScriptStruct->GetStringMetaDataHierarchical(FRigVMStruct::DisplayNameMetaName, &DisplayNameMetadata);
+			DisplayName = DisplayNameMetadata.IsEmpty() ? Trait->GetTraitName() : DisplayNameMetadata;
+
+			const FName DecoratorName = Controller->AddTrait(
+				DecoratorStackNode->GetFName(),
+				*CppDecoratorStruct->GetPathName(),
+				*DisplayName,
+				DefaultValue, INDEX_NONE, true, true);
+			UE_RETURN_ON_ERROR(DecoratorName == DisplayName, TEXT("FAnimationUAFRuntimeTest_InstanceDataGCMissing -> Unexpected trait name"));
+		}
+
+		// Set some values on our trait
+		URigVMPin* DecoratorPin = DecoratorStackNode->FindPin(*DisplayName);
+		UE_RETURN_ON_ERROR(DecoratorPin != nullptr, TEXT("FAnimationUAFRuntimeTest_InstanceDataGCMissing -> Failed to find trait pin"));
+
+		// For some reason a manual recompile is needed, otherwise the trait interface implementations are never run
+		EditorData->RecompileVM();
+
+		TSharedPtr<FAnimNextGraphInstance> GraphInstance = AnimationGraph->AllocateInstance();
+
+		// We place the graph instance in a UObject so that its contained within something the reference collector is aware of
+		// so that FGCGraphInstanceComponent.AddStructReferencedObjects is called. (This isn't strictly needed as it won't traverse
+		// our test trait anyway as it hasn't implemented IGarbageCollection) but its to match parity with the other GC test
+		TObjectPtr<UGraphInstanceHolder> GraphInstanceHolder = NewObject<UGraphInstanceHolder>();
+		GraphInstanceHolder->AddToRoot();
+		GraphInstanceHolder->GraphInstance = GraphInstance;
+		
+		TArray<FString> Messages;
+		FRigVMRuntimeSettings RuntimeSettings;
+		RuntimeSettings.SetLogFunction([&Messages](const FRigVMLogSettings& InLogSettings, const FRigVMExecuteContext* InContext, const FString& Message)
+		{
+			Messages.Add(Message);
+		});
+		GraphInstance->GetComponent<FUAFRigVMComponent>().GetExtendedExecuteContext().SetRuntimeSettings(RuntimeSettings);
+		
+		FUpdateGraphContext UpdateGraphContext(*GraphInstance.Get(), 1.0f / 30.0f);
+		// First update, the anim sequence has been added to root so should definitely be alive
+		UE::UAF::UpdateGraph(UpdateGraphContext);
+
+		CollectGarbage(GARBAGE_COLLECTION_KEEPFLAGS);
+
+		// Second update, the anim sequence should have been destroyed now
+		UE::UAF::UpdateGraph(UpdateGraphContext);
+
+		UE_RETURN_ON_ERROR(Messages.Num() == 2, "FAnimationUAFRuntimeTest_InstanceDataGCMissing -> Not enough messages received");
+		AddErrorIfFalse(Messages[0] == TEXT("Valid"), "FAnimationUAFRuntimeTest_InstanceDataGCMissing -> Update 1: Test AnimSequence should be valid");
+		AddErrorIfFalse(Messages[1] == TEXT("Invalid"), "FAnimationUAFRuntimeTest_InstanceDataGCMissing -> Update 2: Test AnimSequence should be invalid");
+
+		GraphInstanceHolder->GraphInstance = nullptr;
+		GraphInstanceHolder->RemoveFromRoot();
+	}
+
+	Tests::FUtils::CleanupAfterTests();
+
+	return true;
+}
+
+#endif // WITH_DEV_AUTOMATION_TESTS
\ No newline at end of file
diff --git a/Engine/Plugins/Experimental/UAF/UAFTests/Source/UAFAnimGraphTestSuite/Private/GCSharedDataGCTest.cpp b/Engine/Plugins/Experimental/UAF/UAFTests/Source/UAFAnimGraphTestSuite/Private/GCSharedDataGCTest.cpp
new file mode 100644
--- /dev/null
+++ b/Engine/Plugins/Experimental/UAF/UAFTests/Source/UAFAnimGraphTestSuite/Private/GCSharedDataGCTest.cpp
@@ -0,0 +1,218 @@
+// Copyright Epic Games, Inc. All Rights Reserved.
+
+#include "GCTestsUtil.h"
+
+#include "AnimNextRuntimeTest.h"
+#include "AnimNextTest.h"
+#include "UncookedOnlyUtils.h"
+#include "Graph/AnimNextAnimationGraph.h"
+#include "Graph/AnimNextAnimationGraphFactory.h"
+#include "Graph/AnimNextAnimationGraph_EditorData.h"
+#include "Graph/AnimNextGraphInstance.h"
+#include "Graph/RigDecorator_AnimNextCppTrait.h"
+#include "Graph/RigUnit_AnimNextGraphRoot.h"
+#include "Graph/RigUnit_AnimNextTraitStack.h"
+#include "Misc/AutomationTest.h"
+#include "Script/UAFRigVMComponent.h"
+#include "TraitCore/Trait.h"
+#include "TraitCore/TraitRegistry.h"
+#include "TraitInterfaces/IEvaluate.h"
+#include "TraitInterfaces/IGarbageCollection.h"
+#include "TraitInterfaces/IUpdate.h"
+
+#if WITH_DEV_AUTOMATION_TESTS
+
+namespace UE::UAF
+{
+	/** A test trait with a latent UObject property that correctly implements IGarbageCollection */
+	struct FTestTrait_SharedDataGC : FBaseTrait, IUpdate, IGarbageCollection
+	{
+		DECLARE_ANIM_TRAIT(FTestTrait_SharedDataGC, FBaseTrait)
+
+		using FSharedData = FUAFTestAnimSequenceSharedData;
+
+		struct FInstanceData : FTrait::FInstanceData
+		{
+			TWeakObjectPtr<UAnimSequence> WeakAnimSequence;
+			uint8 UpdateCount = 0;
+		};
+		
+		// IGarbageCollection impl
+		virtual void AddReferencedObjects(const FExecutionContext& Context, const TTraitBinding<IGarbageCollection>& Binding, FReferenceCollector& Collector) const override
+		{
+			IGarbageCollection::AddReferencedObjects(Context, Binding, Collector);
+
+			const FSharedData* SharedData = Binding.GetSharedData<FSharedData>();
+
+			TObjectPtr<const UAnimSequence> AnimSequence = SharedData->GetAnimSequence(Binding);
+			Collector.AddReferencedObject(AnimSequence);
+		}
+
+		// IUpdate impl
+		virtual void PostUpdate(FUpdateTraversalContext& Context, const TTraitBinding<IUpdate>& Binding, const FTraitUpdateState& TraitState) const override
+		{
+			IUpdate::PostUpdate(Context, Binding, TraitState);
+
+			FInstanceData* InstanceData = Binding.GetInstanceData<FInstanceData>();
+			const FSharedData* SharedData = Binding.GetSharedData<FSharedData>();
+
+			if (InstanceData->UpdateCount == 0)
+			{
+				InstanceData->WeakAnimSequence = SharedData->GetAnimSequence(Binding);
+				// This anim sequence was added to root so we could verify access to it here during the first update.
+				// Now we'll remove it from root so it's now GC-able before the test calls CollectGarbage.
+				InstanceData->WeakAnimSequence->RemoveFromRoot();
+			}
+			
+			const bool bIsValid = InstanceData->WeakAnimSequence.IsValid() && !InstanceData->WeakAnimSequence->HasAnyFlags(RF_BeginDestroyed | RF_FinishDestroyed);
+			
+			FRigVMExecuteContext& ExecuteContext = Context.GetRootGraphInstance().GetComponent<FUAFRigVMComponent>().GetExtendedExecuteContext().GetPublicData();
+			ExecuteContext.Logf(EMessageSeverity::Info, TEXT("%s"), bIsValid ? TEXT("Valid") : TEXT("Invalid"));
+
+			++InstanceData->UpdateCount;
+		}
+	};
+
+	// Trait implementation boilerplate
+	#define TRAIT_INTERFACE_ENUMERATOR(GeneratorMacro) \
+		GeneratorMacro(IUpdate) \
+		GeneratorMacro(IGarbageCollection) \
+	
+	GENERATE_ANIM_TRAIT_IMPLEMENTATION(FTestTrait_SharedDataGC, TRAIT_INTERFACE_ENUMERATOR, NULL_ANIM_TRAIT_INTERFACE_ENUMERATOR, NULL_ANIM_TRAIT_EVENT_ENUMERATOR)
+	#undef TRAIT_INTERFACE_ENUMERATOR
+}
+
+IMPLEMENT_SIMPLE_AUTOMATION_TEST(FAnimationUAFRuntimeTest_SharedDataGC, "Animation.UAF.Runtime.GCSharedDataWithGCInterface", EAutomationTestFlags::EditorContext | EAutomationTestFlags::EngineFilter)
+
+/*
+ * Tests that a trait that has a latent UObject property must implement IGarbageCollection.
+ * FTestTrait_SharedDataGC correctly implements IGarbageCollection.
+ * Also verifies that the UObject is not GC'd.
+ */
+bool FAnimationUAFRuntimeTest_SharedDataGC::RunTest(const FString& InParameters)
+{
+	using namespace UE::UAF;
+
+	{
+		AUTO_REGISTER_ANIM_TRAIT(FTestTrait_SharedDataGC)
+		FScopedClearNodeTemplateRegistry ScopedClearNodeTemplateRegistry;
+
+		UFactory* GraphFactory = NewObject<UAnimNextAnimationGraphFactory>();
+		TObjectPtr<UAnimNextAnimationGraph> AnimationGraph = CastChecked<UAnimNextAnimationGraph>(GraphFactory->FactoryCreateNew(UAnimNextAnimationGraph::StaticClass(), GetTransientPackage(), TEXT("TestAnimNextGraph"), RF_Transient, nullptr, nullptr, NAME_None));
+		UE_RETURN_ON_ERROR(AnimationGraph != nullptr, "FAnimationUAFRuntimeTest_TraitGC -> Failed to create animation graph");
+
+		TObjectPtr<UAnimNextAnimationGraph_EditorData> EditorData = UncookedOnly::FUtils::GetEditorData<UAnimNextAnimationGraph_EditorData>(AnimationGraph.Get());
+		UE_RETURN_ON_ERROR(EditorData != nullptr, "FAnimationUAFRuntimeTest_TraitGC -> Failed to find module editor data");
+
+		TObjectPtr<UAnimNextController> Controller = Cast<UAnimNextController>(EditorData->GetRigVMClient()->GetController(EditorData->GetRigVMClient()->GetDefaultModel()));
+		UE_RETURN_ON_ERROR(Controller != nullptr, "FAnimationUAFRuntimeTest_TraitGC -> Failed to get RigVM controller");
+
+		// Find graph entry point
+		URigVMNode* MainEntryPointNode = Controller->GetGraph()->FindNodeByName(FRigUnit_AnimNextGraphRoot::StaticStruct()->GetFName());
+		UE_RETURN_ON_ERROR(MainEntryPointNode != nullptr, "FAnimationUAFRuntimeTest_TraitGC -> Failed to find main entry point node");
+
+		URigVMPin* BeginExecutePin = MainEntryPointNode->FindPin(GET_MEMBER_NAME_STRING_CHECKED(FRigUnit_AnimNextGraphRoot, Result));
+		UE_RETURN_ON_ERROR(BeginExecutePin != nullptr && BeginExecutePin->GetDirection() == ERigVMPinDirection::Input, "FAnimationUAFRuntimeTest_TraitGC -> Failed to create entry point");
+
+		URigVMUnitNode* DecoratorStackNode = nullptr;
+		FString DisplayName;
+		{
+			// Suspend auto compilation until we have constructed a valid trait stack
+			TGuardValue<bool> SuspendCompile(EditorData->bAutoRecompileVM, false);
+
+			// Create an empty trait stack node
+			DecoratorStackNode = Controller->AddUnitNode(FRigUnit_AnimNextTraitStack::StaticStruct(), FRigVMStruct::ExecuteName, FVector2D(0.0f, 0.0f), FString(), false);
+			UE_RETURN_ON_ERROR(DecoratorStackNode != nullptr, "FAnimationUAFRuntimeTest_TraitGC -> Failed to create trait stack node");
+
+			// Link our stack result to our entry point
+			Controller->AddLink(DecoratorStackNode->GetPins()[0], MainEntryPointNode->GetPins()[0]);
+
+			// Add a trait
+			const UScriptStruct* CppDecoratorStruct = FRigDecorator_AnimNextCppDecorator::StaticStruct();
+			UE_RETURN_ON_ERROR(CppDecoratorStruct != nullptr, "FAnimationUAFRuntimeTest_TraitGC -> Failed to get find Cpp trait static struct");
+
+			const FTrait* Trait = FTraitRegistry::Get().Find(FTestTrait_SharedDataGC::TraitUID);
+			UE_RETURN_ON_ERROR(Trait != nullptr, "FAnimationUAFRuntimeTest_TraitGC -> Failed to find test trait");
+
+			UScriptStruct* ScriptStruct = Trait->GetTraitSharedDataStruct();
+			UE_RETURN_ON_ERROR(ScriptStruct != nullptr, "FAnimationUAFRuntimeTest_TraitGC -> Failed to find trait shared data struct");
+
+			FString DefaultValue;
+			{
+				const FRigDecorator_AnimNextCppDecorator DefaultCppDecoratorStructInstance;
+				FRigDecorator_AnimNextCppDecorator CppDecoratorStructInstance;
+				CppDecoratorStructInstance.DecoratorSharedDataStruct = ScriptStruct;
+
+				UE_RETURN_ON_ERROR(CppDecoratorStructInstance.CanBeAddedToNode(DecoratorStackNode, nullptr), "FAnimationUAFRuntimeTest_TraitGC -> Trait cannot be added to trait stack node");
+
+				FRigDecorator_AnimNextCppDecorator::StaticStruct()->ExportText(DefaultValue, &CppDecoratorStructInstance, &DefaultCppDecoratorStructInstance, nullptr, PPF_None, nullptr);
+			}
+
+			FString DisplayNameMetadata;
+			ScriptStruct->GetStringMetaDataHierarchical(FRigVMStruct::DisplayNameMetaName, &DisplayNameMetadata);
+			DisplayName = DisplayNameMetadata.IsEmpty() ? Trait->GetTraitName() : DisplayNameMetadata;
+
+			const FName DecoratorName = Controller->AddTrait(
+				DecoratorStackNode->GetFName(),
+				*CppDecoratorStruct->GetPathName(),
+				*DisplayName,
+				DefaultValue, INDEX_NONE, true, true);
+			UE_RETURN_ON_ERROR(DecoratorName == DisplayName, TEXT("FAnimationUAFRuntimeTest_TraitGC -> Unexpected trait name"));
+		}
+
+		// Set some values on our trait
+		URigVMPin* DecoratorPin = DecoratorStackNode->FindPin(*DisplayName);
+		UE_RETURN_ON_ERROR(DecoratorPin != nullptr, TEXT("FAnimationUAFRuntimeTest_TraitGC -> Failed to find trait pin"));
+
+		{
+			URigVMUnitNode* TestNode = Controller->AddUnitNode(FRigUnit_AnimSequenceTest::StaticStruct());
+			UE_RETURN_ON_ERROR(TestNode != nullptr, TEXT("FAnimationUAFRuntimeTest_TraitGC -> Failed to add test node"));
+
+			const bool bSuccess = Controller->AddLink(
+				TestNode->FindPin(GET_MEMBER_NAME_STRING_CHECKED(FRigUnit_AnimSequenceTest, AnimSequence)),
+				DecoratorPin->GetSubPins()[1]);	// TestAnimSequence
+
+			UE_RETURN_ON_ERROR(bSuccess, TEXT("FAnimationUAFRuntimeTest_TraitGC -> Failed to add link"));
+		}
+
+		TSharedPtr<FAnimNextGraphInstance> GraphInstance = AnimationGraph->AllocateInstance();
+
+		// We place the graph instance in a UObject so that its contained within something the reference collector is aware of
+		TObjectPtr<UGraphInstanceHolder> GraphInstanceHolder = NewObject<UGraphInstanceHolder>();
+		GraphInstanceHolder->AddToRoot();
+		GraphInstanceHolder->GraphInstance = GraphInstance;
+
+		TArray<FString> Messages;
+		FRigVMRuntimeSettings RuntimeSettings;
+		RuntimeSettings.SetLogFunction([&Messages](const FRigVMLogSettings& InLogSettings, const FRigVMExecuteContext* InContext, const FString& Message)
+		{
+			Messages.Add(Message);
+		});
+		GraphInstance->GetComponent<FUAFRigVMComponent>().GetExtendedExecuteContext().SetRuntimeSettings(RuntimeSettings);
+		
+		FUpdateGraphContext UpdateGraphContext(*GraphInstance.Get(), 1.0f / 30.0f);
+		
+		// First update, the anim sequence has been added to root so should definitely be alive
+		// Our test trait should report it as valid in its IUpdate
+		UE::UAF::UpdateGraph(UpdateGraphContext);
+
+		CollectGarbage(GARBAGE_COLLECTION_KEEPFLAGS);
+
+		// Second update, the anim sequence should still be alive
+		// Our test trait should report it as still valid in its IUpdate
+		UE::UAF::UpdateGraph(UpdateGraphContext);
+
+		UE_RETURN_ON_ERROR(Messages.Num() == 2, "FAnimationUAFRuntimeTest_TraitGC -> Not enough messages received");
+		AddErrorIfFalse(Messages[0] == TEXT("Valid"), "FAnimationUAFRuntimeTest_TraitGC -> Update 1: Test AnimSequence should be valid");
+		AddErrorIfFalse(Messages[1] == TEXT("Valid"), "FAnimationUAFRuntimeTest_TraitGC -> Update 2: Test AnimSequence should be valid");
+
+		GraphInstanceHolder->GraphInstance = nullptr;
+		GraphInstanceHolder->RemoveFromRoot();
+	}
+
+	Tests::FUtils::CleanupAfterTests();
+
+	return true;
+}
+
+#endif // WITH_DEV_AUTOMATION_TESTS
\ No newline at end of file
diff --git a/Engine/Plugins/Experimental/UAF/UAFTests/Source/UAFAnimGraphTestSuite/Private/GCSharedDataMissingGCTest.cpp b/Engine/Plugins/Experimental/UAF/UAFTests/Source/UAFAnimGraphTestSuite/Private/GCSharedDataMissingGCTest.cpp
new file mode 100644
--- /dev/null
+++ b/Engine/Plugins/Experimental/UAF/UAFTests/Source/UAFAnimGraphTestSuite/Private/GCSharedDataMissingGCTest.cpp
@@ -0,0 +1,142 @@
+// Copyright Epic Games, Inc. All Rights Reserved.
+
+#include "GCTestsUtil.h"
+
+#include "AnimNextRuntimeTest.h"
+#include "AnimNextTest.h"
+#include "UncookedOnlyUtils.h"
+#include "Graph/AnimNextAnimationGraph.h"
+#include "Graph/AnimNextAnimationGraphFactory.h"
+#include "Graph/AnimNextAnimationGraph_EditorData.h"
+#include "Graph/AnimNextGraphInstance.h"
+#include "Graph/RigDecorator_AnimNextCppTrait.h"
+#include "Graph/RigUnit_AnimNextGraphRoot.h"
+#include "Graph/RigUnit_AnimNextTraitStack.h"
+#include "Misc/AutomationTest.h"
+#include "Script/UAFRigVMComponent.h"
+#include "TraitCore/Trait.h"
+#include "TraitCore/TraitRegistry.h"
+#include "TraitInterfaces/IEvaluate.h"
+#include "TraitInterfaces/IUpdate.h"
+
+#if WITH_DEV_AUTOMATION_TESTS
+
+namespace UE::UAF
+{
+	/** A test trait with a latent UObject property that incorrectly does not implement IGarbageCollection */
+	struct FTestTrait_SharedDataMissingGC : FBaseTrait
+	{
+		DECLARE_ANIM_TRAIT(FTestTrait_MissingGC, FBaseTrait)
+
+		using FSharedData = FUAFTestAnimSequenceSharedData;
+	};
+
+	GENERATE_ANIM_TRAIT_IMPLEMENTATION(FTestTrait_SharedDataMissingGC, NULL_ANIM_TRAIT_INTERFACE_ENUMERATOR, NULL_ANIM_TRAIT_INTERFACE_ENUMERATOR, NULL_ANIM_TRAIT_EVENT_ENUMERATOR)
+}
+
+IMPLEMENT_SIMPLE_AUTOMATION_TEST(FAnimationUAFRuntimeTest_SharedDataMissingGC, "Animation.UAF.Runtime.GCSharedDataMissingGCInterface", EAutomationTestFlags::EditorContext | EAutomationTestFlags::EngineFilter)
+
+/*
+ * Tests that a trait that has a latent UObject property must implement IGarbageCollection.
+ * FTestTrait_SharedDataMissingGC incorrectly doesn't implement IGarbageCollection which should report a compile error.
+ * Also verifies that the UObject is GC'd correctly in this case where it has no references.
+ */
+bool FAnimationUAFRuntimeTest_SharedDataMissingGC::RunTest(const FString& InParameters)
+{
+	using namespace UE::UAF;
+
+	{
+		AUTO_REGISTER_ANIM_TRAIT(FTestTrait_SharedDataMissingGC)
+		FScopedClearNodeTemplateRegistry ScopedClearNodeTemplateRegistry;
+
+		AddExpectedError(TEXT("Trait is connected to a UObject but does not implement IGarbageCollection"), EAutomationExpectedErrorFlags::Contains, 6);
+
+		UFactory* GraphFactory = NewObject<UAnimNextAnimationGraphFactory>();
+		UAnimNextAnimationGraph* AnimationGraph = CastChecked<UAnimNextAnimationGraph>(GraphFactory->FactoryCreateNew(UAnimNextAnimationGraph::StaticClass(), GetTransientPackage(), TEXT("TestAnimNextGraph"), RF_Transient, nullptr, nullptr, NAME_None));
+		UE_RETURN_ON_ERROR(AnimationGraph != nullptr, "FAnimationUAFRuntimeTest_TraitMissingGC -> Failed to create animation graph");
+
+		UAnimNextAnimationGraph_EditorData* EditorData = UncookedOnly::FUtils::GetEditorData<UAnimNextAnimationGraph_EditorData>(AnimationGraph);
+		UE_RETURN_ON_ERROR(EditorData != nullptr, "FAnimationUAFRuntimeTest_TraitMissingGC -> Failed to find module editor data");
+
+		UAnimNextController* Controller = Cast<UAnimNextController>(EditorData->GetRigVMClient()->GetController(EditorData->GetRigVMClient()->GetDefaultModel()));
+		UE_RETURN_ON_ERROR(Controller != nullptr, "FAnimationUAFRuntimeTest_TraitMissingGC -> Failed to get RigVM controller");
+
+		// Find graph entry point
+		URigVMNode* MainEntryPointNode = Controller->GetGraph()->FindNodeByName(FRigUnit_AnimNextGraphRoot::StaticStruct()->GetFName());
+		UE_RETURN_ON_ERROR(MainEntryPointNode != nullptr, "FAnimationUAFRuntimeTest_TraitMissingGC -> Failed to find main entry point node");
+
+		URigVMPin* BeginExecutePin = MainEntryPointNode->FindPin(GET_MEMBER_NAME_STRING_CHECKED(FRigUnit_AnimNextGraphRoot, Result));
+		UE_RETURN_ON_ERROR(BeginExecutePin != nullptr && BeginExecutePin->GetDirection() == ERigVMPinDirection::Input, "FAnimationUAFRuntimeTest_TraitMissingGC -> Failed to create entry point");
+
+		URigVMUnitNode* DecoratorStackNode = nullptr;
+		FString DisplayName;
+		{
+			// Suspend auto compilation until we have constructed a valid trait stack
+			TGuardValue<bool> SuspendCompile(EditorData->bAutoRecompileVM, false);
+
+			// Create an empty trait stack node
+			DecoratorStackNode = Controller->AddUnitNode(FRigUnit_AnimNextTraitStack::StaticStruct(), FRigVMStruct::ExecuteName, FVector2D(0.0f, 0.0f), FString(), false);
+			UE_RETURN_ON_ERROR(DecoratorStackNode != nullptr, "FAnimationUAFRuntimeTest_TraitMissingGC -> Failed to create trait stack node");
+
+			// Link our stack result to our entry point
+			Controller->AddLink(DecoratorStackNode->GetPins()[0], MainEntryPointNode->GetPins()[0]);
+
+			// Add a trait
+			const UScriptStruct* CppDecoratorStruct = FRigDecorator_AnimNextCppDecorator::StaticStruct();
+			UE_RETURN_ON_ERROR(CppDecoratorStruct != nullptr, "FAnimationUAFRuntimeTest_TraitMissingGC -> Failed to get find Cpp trait static struct");
+
+			const FTrait* Trait = FTraitRegistry::Get().Find(FTestTrait_SharedDataMissingGC::TraitUID);
+			UE_RETURN_ON_ERROR(Trait != nullptr, "FAnimationUAFRuntimeTest_TraitMissingGC -> Failed to find test trait");
+
+			UScriptStruct* ScriptStruct = Trait->GetTraitSharedDataStruct();
+			UE_RETURN_ON_ERROR(ScriptStruct != nullptr, "FAnimationUAFRuntimeTest_TraitMissingGC -> Failed to find trait shared data struct");
+
+			FString DefaultValue;
+			{
+				const FRigDecorator_AnimNextCppDecorator DefaultCppDecoratorStructInstance;
+				FRigDecorator_AnimNextCppDecorator CppDecoratorStructInstance;
+				CppDecoratorStructInstance.DecoratorSharedDataStruct = ScriptStruct;
+
+				UE_RETURN_ON_ERROR(CppDecoratorStructInstance.CanBeAddedToNode(DecoratorStackNode, nullptr), "FAnimationUAFRuntimeTest_TraitMissingGC -> Trait cannot be added to trait stack node");
+
+				FRigDecorator_AnimNextCppDecorator::StaticStruct()->ExportText(DefaultValue, &CppDecoratorStructInstance, &DefaultCppDecoratorStructInstance, nullptr, PPF_None, nullptr);
+			}
+
+			FString DisplayNameMetadata;
+			ScriptStruct->GetStringMetaDataHierarchical(FRigVMStruct::DisplayNameMetaName, &DisplayNameMetadata);
+			DisplayName = DisplayNameMetadata.IsEmpty() ? Trait->GetTraitName() : DisplayNameMetadata;
+
+			const FName DecoratorName = Controller->AddTrait(
+				DecoratorStackNode->GetFName(),
+				*CppDecoratorStruct->GetPathName(),
+				*DisplayName,
+				DefaultValue, INDEX_NONE, true, true);
+			UE_RETURN_ON_ERROR(DecoratorName == DisplayName, TEXT("FAnimationUAFRuntimeTest_TraitMissingGC -> Unexpected trait name"));
+		}
+
+		// Set some values on our trait
+		URigVMPin* DecoratorPin = DecoratorStackNode->FindPin(*DisplayName);
+		UE_RETURN_ON_ERROR(DecoratorPin != nullptr, TEXT("FAnimationUAFRuntimeTest_TraitMissingGC -> Failed to find trait pin"));
+
+		{
+			URigVMUnitNode* TestNode = Controller->AddUnitNode(FRigUnit_AnimSequenceTest::StaticStruct());
+			UE_RETURN_ON_ERROR(TestNode != nullptr, TEXT("FAnimationUAFRuntimeTest_TraitMissingGC -> Failed to add test node"));
+
+			const bool bSuccess = Controller->AddLink(
+				TestNode->FindPin(GET_MEMBER_NAME_STRING_CHECKED(FRigUnit_AnimSequenceTest, AnimSequence)),
+				DecoratorPin->GetSubPins()[1]);	// TestAnimSequence
+
+			UE_RETURN_ON_ERROR(TestNode != nullptr, TEXT("FAnimationUAFRuntimeTest_TraitMissingGC -> Failed to add link"));
+		}
+
+		TSharedPtr<FAnimNextGraphInstance> GraphInstance = AnimationGraph->AllocateInstance();
+
+		// No need to update the graph, the graph will not compile with this invalid trait
+	}
+
+	Tests::FUtils::CleanupAfterTests();
+
+	return true;
+}
+
+#endif // WITH_DEV_AUTOMATION_TESTS
\ No newline at end of file
diff --git a/Engine/Plugins/Experimental/UAF/UAFTests/Source/UAFAnimGraphTestSuite/Private/GCTestsUtil.cpp b/Engine/Plugins/Experimental/UAF/UAFTests/Source/UAFAnimGraphTestSuite/Private/GCTestsUtil.cpp
new file mode 100644
--- /dev/null
+++ b/Engine/Plugins/Experimental/UAF/UAFTests/Source/UAFAnimGraphTestSuite/Private/GCTestsUtil.cpp
@@ -0,0 +1,28 @@
+// Copyright Epic Games, Inc. All Rights Reserved.
+
+#include "GCTestsUtil.h"
+
+#include "Graph/AnimNextAnimationGraph.h"
+#include "Graph/AnimNextAnimationGraphFactory.h"
+#include "Misc/AutomationTest.h"
+
+FRigUnit_AnimSequenceTest_Execute()
+{
+	if (!bCreated)
+	{
+		AnimSequence = NewObject<UAnimSequence>(GetTransientPackage(), NAME_None);
+		AnimSequence->AddToRoot(); // Will be removed during the test
+		bCreated = true;
+	}
+}
+
+void UGraphInstanceHolder::AddReferencedObjects(UObject* InThis, FReferenceCollector& Collector) 
+{
+	UGraphInstanceHolder* This = CastChecked<UGraphInstanceHolder>(InThis);
+	if (This->GraphInstance.IsValid())
+	{
+		Collector.AddPropertyReferencesWithStructARO(FAnimNextGraphInstance::StaticStruct(), This->GraphInstance.Get(), This);
+	}
+
+	Super::AddReferencedObjects(InThis, Collector);
+}
\ No newline at end of file
diff --git a/Engine/Plugins/Experimental/UAF/UAFTests/Source/UAFAnimGraphTestSuite/Private/GCTestsUtil.h b/Engine/Plugins/Experimental/UAF/UAFTests/Source/UAFAnimGraphTestSuite/Private/GCTestsUtil.h
new file mode 100644
--- /dev/null
+++ b/Engine/Plugins/Experimental/UAF/UAFTests/Source/UAFAnimGraphTestSuite/Private/GCTestsUtil.h
@@ -0,0 +1,52 @@
+// Copyright Epic Games, Inc. All Rights Reserved.
+
+#pragma once
+
+#include "TraitCore/TraitBinding.h"
+#include "TraitCore/TraitSharedData.h"
+#include "Animation/AnimSequence.h"
+#include "Graph/AnimNextGraphInstance.h"
+#include "Graph/RigUnit_AnimNextBase.h"
+
+#include "GCTestsUtil.generated.h"
+
+USTRUCT()
+struct FUAFTestAnimSequenceSharedData : public FAnimNextTraitSharedData
+{
+	GENERATED_BODY()
+
+	UPROPERTY(EditAnywhere, Category = "Default")
+	TObjectPtr<UAnimSequence> AnimSequence;
+
+	#define TRAIT_LATENT_PROPERTIES_ENUMERATOR(GeneratorMacro) \
+		GeneratorMacro(AnimSequence) \
+
+	GENERATE_TRAIT_LATENT_PROPERTIES(FUAFTestAnimSequenceSharedData, TRAIT_LATENT_PROPERTIES_ENUMERATOR)
+	#undef TRAIT_LATENT_PROPERTIES_ENUMERATOR
+};
+
+USTRUCT()
+struct FRigUnit_AnimSequenceTest : public FRigUnit_AnimNextBase
+{
+	GENERATED_BODY()
+	
+	RIGVM_METHOD()
+	void Execute();
+
+	UPROPERTY(EditAnywhere, Category = "Default", meta = (Output))
+	TObjectPtr<UAnimSequence> AnimSequence;
+
+	UPROPERTY()
+	bool bCreated = false;
+};
+
+UCLASS()
+class UGraphInstanceHolder : public UObject
+{
+	GENERATED_BODY()
+
+public:
+	TSharedPtr<FAnimNextGraphInstance> GraphInstance;
+
+	static void AddReferencedObjects(UObject* InThis, FReferenceCollector& Collector);
+};
\ No newline at end of file
