RigVMClient isn't exposed to uaf_system in Python API
#fix UAF Test plugin scripting access to subset of IRigVMClientHost API on UAnimNextRigVMAssetEditorData
#misc add PythonAutomation plugin as dependency for UAFTest plugin

diff --git a/Engine/Plugins/Experimental/UAF/UAF/Source/UAFUncookedOnly/Internal/AnimNextRigVMAssetEditorData.h b/Engine/Plugins/Experimental/UAF/UAF/Source/UAFUncookedOnly/Internal/AnimNextRigVMAssetEditorData.h
--- a/Engine/Plugins/Experimental/UAF/UAF/Source/UAFUncookedOnly/Internal/AnimNextRigVMAssetEditorData.h
+++ b/Engine/Plugins/Experimental/UAF/UAF/Source/UAFUncookedOnly/Internal/AnimNextRigVMAssetEditorData.h
@@ -22,6 +22,7 @@ class UAnimNextRigVMAssetEditorData;
 class UAnimNextEdGraph;
 class UAnimNextSharedVariablesEntry;
 class UAnimNextWorkspaceEditorMode;
+class UUAFTestBlueprintLibrary;
 struct FAnimNextRigVMAssetCompileContext;
 struct FAnimNextGetFunctionHeaderCompileContext;
 struct FAnimNextGetVariableCompileContext;
@@ -255,6 +256,7 @@ protected:
 	friend class UE::UAF::Editor::SAddVariablesDialog;
 	friend class UE::UAF::Editor::FVariableProxyCustomization;
 	friend class UAnimNextSharedVariablesEntry;
+	friend class UUAFTestBlueprintLibrary;
 	friend struct UE::UAF::UncookedOnly::FPublicVariablesImpl;
 	friend UE::UAF::Editor::FAnimNextAnimGraphEditorModule;
 	friend class UE::UAF::Editor::FCallFunctionSharedDataDetails;
diff --git a/Engine/Plugins/Experimental/UAF/UAFTests/Source/UAFTests/Private/UAFTestBlueprintLibrary.h b/Engine/Plugins/Experimental/UAF/UAFTests/Source/UAFTests/Private/UAFTestBlueprintLibrary.h
new file mode 100644
--- /dev/null
+++ b/Engine/Plugins/Experimental/UAF/UAFTests/Source/UAFTests/Private/UAFTestBlueprintLibrary.h
@@ -0,0 +1,179 @@
+// Copyright Epic Games, Inc. All Rights Reserved.
+
+#pragma once
+
+#include "CoreMinimal.h"
+#include "Kismet/BlueprintFunctionLibrary.h"
+#include "AnimNextRigVMAssetEditorData.h"
+#include "UAFTestBlueprintLibrary.generated.h"
+
+UCLASS()
+class UUAFTestBlueprintLibrary : public UBlueprintFunctionLibrary
+{
+	GENERATED_BODY()
+	
+public:
+		
+	UFUNCTION(BlueprintCallable, Category = "UAF|Test", meta=(ScriptMethod))
+	static void RecompileVM(UAnimNextRigVMAsset* InAsset) 
+	{	
+		if (InAsset)
+		{
+			UE::UAF::UncookedOnly::FUtils::GetEditorData(InAsset)->RecompileVM(); 
+		}
+		else
+		{
+			UE_LOG(LogAnimation, Warning, TEXT("Invalid Asset in UUAFTestBlueprintLibrary::RecompileVM"))
+		}
+	}
+
+	UFUNCTION(BlueprintCallable, Category = "UAF|Test", meta=(ScriptMethod))
+	static URigVMGraph* GetModel(UAnimNextRigVMAsset* InAsset, const UEdGraph* InEdGraph = nullptr) 
+	{
+		if (InAsset)
+		{
+			return UE::UAF::UncookedOnly::FUtils::GetEditorData(InAsset)->GetModel(); 
+		}
+		else
+		{
+			UE_LOG(LogAnimation, Warning, TEXT("Invalid Asset in UUAFTestBlueprintLibrary::GetModel"))
+		}
+		
+		return nullptr;
+	}
+
+	UFUNCTION(BlueprintCallable, Category = "UAF|Test", meta=(ScriptMethod))
+	static URigVMGraph* GetDefaultModel(UAnimNextRigVMAsset* InAsset)
+	{
+		if (InAsset)
+		{
+			return UE::UAF::UncookedOnly::FUtils::GetEditorData(InAsset)->GetDefaultModel(); 
+		}
+		else
+		{
+			UE_LOG(LogAnimation, Warning, TEXT("Invalid Asset in UUAFTestBlueprintLibrary::GetDefaultModel"))
+		}
+		
+		return nullptr;
+	}
+
+	UFUNCTION(BlueprintCallable, Category = "UAF|Test", meta=(ScriptMethod))
+	static TArray<URigVMGraph*> GetAllModels(UAnimNextRigVMAsset* InAsset)
+	{
+		if (InAsset)
+		{
+			return UE::UAF::UncookedOnly::FUtils::GetEditorData(InAsset)->GetAllModels(); 
+		}
+		else
+		{
+			UE_LOG(LogAnimation, Warning, TEXT("Invalid Asset in UUAFTestBlueprintLibrary::GetAllModels"))
+		}
+
+		return TArray<URigVMGraph*>();
+	}
+
+	UFUNCTION(BlueprintCallable, Category = "UAF|Test", meta=(ScriptMethod))
+	static URigVMFunctionLibrary* GetLocalFunctionLibrary(UAnimNextRigVMAsset* InAsset)
+	{
+		if (InAsset)
+		{
+			return UE::UAF::UncookedOnly::FUtils::GetEditorData(InAsset)->GetLocalFunctionLibrary(); 
+		}
+		else
+		{
+			UE_LOG(LogAnimation, Warning, TEXT("Invalid Asset in UUAFTestBlueprintLibrary::GetLocalFunctionLibrary"))
+		}
+
+		return nullptr;
+	}
+
+	UFUNCTION(BlueprintCallable, Category = "UAF|Test", meta=(ScriptMethod))
+	static URigVMFunctionLibrary* GetOrCreateLocalFunctionLibrary(UAnimNextRigVMAsset* InAsset, bool bSetupUndoRedo = true)
+	{
+		if (InAsset)
+		{
+			return UE::UAF::UncookedOnly::FUtils::GetEditorData(InAsset)->GetOrCreateLocalFunctionLibrary(bSetupUndoRedo); 
+		}
+		else
+		{
+			UE_LOG(LogAnimation, Warning, TEXT("Invalid Asset in UUAFTestBlueprintLibrary::GetOrCreateLocalFunctionLibrary"))
+		}
+
+		return nullptr;
+	}
+
+	UFUNCTION(BlueprintCallable, Category = "UAF|Test", meta=(ScriptMethod))
+	static URigVMGraph* AddModel(UAnimNextRigVMAsset* InAsset, FString InName = TEXT("Rig Graph"), bool bSetupUndoRedo = true, bool bPrintPythonCommand = true)
+	{
+		if (InAsset)
+		{
+			return UE::UAF::UncookedOnly::FUtils::GetEditorData(InAsset)->AddModel(InName, bSetupUndoRedo, bPrintPythonCommand); 
+		}
+		else
+		{
+			UE_LOG(LogAnimation, Warning, TEXT("Invalid Asset in UUAFTestBlueprintLibrary::AddModel"))
+		}
+
+		return nullptr;
+	}
+
+	UFUNCTION(BlueprintCallable, Category = "UAF|Test", meta=(ScriptMethod))
+	static bool RemoveModel(UAnimNextRigVMAsset* InAsset, FString InName = TEXT("Rig Graph"), bool bSetupUndoRedo = true, bool bPrintPythonCommand = true)
+	{
+		if (InAsset)
+		{
+			return UE::UAF::UncookedOnly::FUtils::GetEditorData(InAsset)->RemoveModel(InName, bSetupUndoRedo, bPrintPythonCommand); 
+		}
+		else
+		{
+			UE_LOG(LogAnimation, Warning, TEXT("Invalid Asset in UUAFTestBlueprintLibrary::RemoveModel"))
+		}
+
+		return false;
+	}
+
+	UFUNCTION(BlueprintCallable, Category = "UAF|Test", meta=(ScriptMethod))
+	static URigVMController* GetController(UAnimNextRigVMAsset* InAsset, const URigVMGraph* InGraph = nullptr)
+	{
+		if (InAsset)
+		{
+			return UE::UAF::UncookedOnly::FUtils::GetEditorData(InAsset)->GetController(InGraph); 
+		}
+		else
+		{
+			UE_LOG(LogAnimation, Warning, TEXT("Invalid Asset in UUAFTestBlueprintLibrary::GetController"))
+		}
+
+		return nullptr;
+	}
+
+	UFUNCTION(BlueprintCallable, Category = "UAF|Test", meta=(ScriptMethod))
+	static URigVMController* GetControllerByName(UAnimNextRigVMAsset* InAsset, const FString InGraphName = TEXT(""))
+	{
+		if (InAsset)
+		{
+			return UE::UAF::UncookedOnly::FUtils::GetEditorData(InAsset)->GetControllerByName(InGraphName); 
+		}
+		else
+		{
+			UE_LOG(LogAnimation, Warning, TEXT("Invalid Asset in UUAFTestBlueprintLibrary::GetControllerByName"))
+		}
+
+		return nullptr;
+	}
+
+	UFUNCTION(BlueprintCallable, Category = "UAF|Test", meta=(ScriptMethod))
+	static URigVMController* GetOrCreateController(UAnimNextRigVMAsset* InAsset, URigVMGraph* InGraph = nullptr)
+	{
+		if (InAsset)
+		{
+			return UE::UAF::UncookedOnly::FUtils::GetEditorData(InAsset)->GetController(InGraph); 
+		}
+		else
+		{
+			UE_LOG(LogAnimation, Warning, TEXT("Invalid Asset in UUAFTestBlueprintLibrary::GetOrCreateController"))
+		}
+
+		return nullptr;
+	}
+};
diff --git a/Engine/Plugins/Experimental/UAF/UAFTests/UAFTests.uplugin b/Engine/Plugins/Experimental/UAF/UAFTests/UAFTests.uplugin
--- a/Engine/Plugins/Experimental/UAF/UAFTests/UAFTests.uplugin
+++ b/Engine/Plugins/Experimental/UAF/UAFTests/UAFTests.uplugin
@@ -51,6 +51,10 @@
 		{
 			"Name": "RigVM",
 			"Enabled": true		
+		},
+		{
+			"Name": "PythonAutomationTest",
+			"Enabled": true
 		}
 	]
 }
