Introducing new abstract skeleton runtime for UAF (disabled by default)

diff --git a/Engine/Plugins/Experimental/UAF/UAFAnimGraph/Source/UAFAnimGraph/Public/EvaluationVM/EvaluationVM.h b/Engine/Plugins/Experimental/UAF/UAFAnimGraph/Source/UAFAnimGraph/Public/EvaluationVM/EvaluationVM.h
--- a/Engine/Plugins/Experimental/UAF/UAFAnimGraph/Source/UAFAnimGraph/Public/EvaluationVM/EvaluationVM.h
+++ b/Engine/Plugins/Experimental/UAF/UAFAnimGraph/Source/UAFAnimGraph/Public/EvaluationVM/EvaluationVM.h
@@ -11,6 +11,11 @@
 #include "Hash/Fnv.h"
 #include "ReferencePose.h"
 #include "Misc/Crc.h"
+#include "UAF/AbstractSkeleton/AbstractSkeletonSetBinding.h"
+#include "UAF/Attributes/AttributeCollection.h"
+#include "UAF/Attributes/AttributeEvaluationContext.h"
+#include "UAF/Attributes/AttributeRuntimeRegistry.h"
+#include "UAF/Attributes/Transformers/Transformer.h"
 
 #define UE_API UAFANIMGRAPH_API
 
@@ -147,7 +152,11 @@ namespace UE::UAF
 	struct FEvaluationVM
 	{
 		FEvaluationVM() = default;
-		UE_API FEvaluationVM(EEvaluationFlags InEvaluationFlags, const UE::UAF::FReferencePose& InReferencePose, int32 InCurrentLOD);
+		UE_API FEvaluationVM(EEvaluationFlags InEvaluationFlags, const FReferencePose& InReferencePose, int32 InCurrentLOD);
+		UE_API FEvaluationVM(EEvaluationFlags EvaluationFlags, const UAbstractSkeletonSetBinding& SetBinding, FName SetName, int32 CurrentLOD);
+
+		// REMOVE ME: Exclusively for testing that before/after yields the same results
+		UE_API FEvaluationVM(EEvaluationFlags InEvaluationFlags, const FReferencePose& IReferencePose, const UAbstractSkeletonSetBinding* SetBinding, FName SetName, int32 InCurrentLOD);
 
 		// Allow move
 		FEvaluationVM(FEvaluationVM&&) = default;
@@ -189,10 +198,18 @@ namespace UE::UAF
 		// Returns an uninitialized keyframe with memory pre-allocated
 		[[nodiscard]] UE_API FKeyframeState MakeUninitializedKeyframe(bool bAdditiveKeyframe) const;
 
-		int32 GetCurrentLOD() const
-		{
-			return CurrentLOD;
-		}
+		// Returns an attribute collection sized for the current evaluation mask and LOD
+		[[nodiscard]] UE_API FAttributeCollectionStack MakeAttributeCollection(bool bAdditive, bool bInitialize) const;
+
+		// Returns the VM context that controls what state we produce
+		[[nodiscard]] UE_API const FAttributeEvaluationContext& GetContext() const;
+
+		// Returns the current LOD
+		[[nodiscard]] int32 GetCurrentLOD() const { return CurrentLOD; }
+
+		// Returns the attribute value transformer map
+		[[nodiscard]] UE_API const FAttributeValueTransformerMapPtr& GetTransformerMap() const;
+
 	private:
 		// Disallow copy
 		FEvaluationVM(const FEvaluationVM&) = delete;
@@ -211,7 +228,14 @@ namespace UE::UAF
 		TMap<FName, FEvaluationVMStack> InternalStacks;
 
 		// Reference pose
-		const UE::UAF::FReferencePose* ReferencePose = nullptr;
+		const FReferencePose* ReferencePose = nullptr;
+
+		// The internal context that controls what state we are producing
+		// TODO: Should be a stack, push/pop a set mask
+		FAttributeEvaluationContext InternalContext;
+
+		// Shared map of attribute value transformers
+		FAttributeValueTransformerMapPtr AttributeValueTransformerMap;
 
 		// Current LOD we are evaluating at
 		int32 CurrentLOD = 0;
@@ -306,6 +330,10 @@ namespace UE::UAF
 	// A stack of FKeyframeState instances used when sampling sequences and blending their results
 	extern UAFANIMGRAPH_API const FEvaluationVMStackName KEYFRAME_STACK_NAME;
 	ANIM_NEXT_ENABLE_EVALUATION_STACK_USAGE(TUniquePtr<FKeyframeState>)
+
+	// A stack of FAttributeCollectionStack instances used when sampling sequences and blending their results
+	extern UAFANIMGRAPH_API const FEvaluationVMStackName ATTRIBUTE_STACK_NAME;
+	ANIM_NEXT_ENABLE_EVALUATION_STACK_USAGE(TUniquePtr<FAttributeCollection>)
 }
 
 #undef UE_API
diff --git a/Engine/Plugins/Experimental/UAF/UAFAnimGraph/Source/UAFAnimGraph/Public/EvaluationVM/Tasks/BlendKeyframes.h b/Engine/Plugins/Experimental/UAF/UAFAnimGraph/Source/UAFAnimGraph/Public/EvaluationVM/Tasks/BlendKeyframes.h
--- a/Engine/Plugins/Experimental/UAF/UAFAnimGraph/Source/UAFAnimGraph/Public/EvaluationVM/Tasks/BlendKeyframes.h
+++ b/Engine/Plugins/Experimental/UAF/UAFAnimGraph/Source/UAFAnimGraph/Public/EvaluationVM/Tasks/BlendKeyframes.h
@@ -13,6 +13,7 @@ struct FInputScaleBiasClamp;
 
 namespace UE::UAF
 {
+	class FAttributeCollection;
 	struct FKeyframeState;
 }
 
@@ -52,6 +53,7 @@ struct FAnimNextBlendTwoKeyframesTask : public FAnimNextEvaluationTask
 
 	// Compute the current alpha value for curves or return the interpolation alpha otherwise
 	UE_API float GetInterpolationAlpha(const UE::UAF::FKeyframeState* KeyframeA, const UE::UAF::FKeyframeState* KeyframeB) const;
+	UE_API float GetInterpolationAlpha(const UE::UAF::FAttributeCollection& CollectionA, const UE::UAF::FAttributeCollection& CollectionB) const;
 
 	// The interpolation alpha between the two input keyframes
 	UPROPERTY(VisibleAnywhere, Category = Properties)
diff --git a/Engine/Plugins/Experimental/UAF/UAFAnimGraph/Source/UAFAnimGraph/Public/TraitInterfaces/IEvaluate.h b/Engine/Plugins/Experimental/UAF/UAFAnimGraph/Source/UAFAnimGraph/Public/TraitInterfaces/IEvaluate.h
--- a/Engine/Plugins/Experimental/UAF/UAFAnimGraph/Source/UAFAnimGraph/Public/TraitInterfaces/IEvaluate.h
+++ b/Engine/Plugins/Experimental/UAF/UAFAnimGraph/Source/UAFAnimGraph/Public/TraitInterfaces/IEvaluate.h
@@ -15,6 +15,7 @@ struct FAnimNextGraphLODPose;
 struct FRigUnit_AnimNextRunAnimationGraph_v1;
 struct FRigUnit_AnimNextRunAnimationGraph_v2;
 struct FAnimNode_AnimNextGraph;
+class UAbstractSkeletonSetBinding;
 
 namespace UE::UAF
 {
@@ -38,16 +39,32 @@ namespace UE::UAF
 		const FReferencePose& GetRefPose() const { return RefPose; }
 		int32 GetGraphLODLevel() const { return GraphLODLevel; }
 
+		// Only set if the new abstract skeleton runtime is enabled
+		UAbstractSkeletonSetBinding* GetSkeletonSetBinding() const { return SkeletonSetBinding; }
+		FName GetSkeletonSetName() const { return SkeletonSetName; }
+
 	protected:
 		const TWeakObjectPtr<const USkeletalMeshComponent>& GetBindingObject() const { return BindingObject; }
 		void SetBindingObject(const TWeakObjectPtr<const USkeletalMeshComponent>& InBindingObject) { BindingObject = InBindingObject; }
 
+		void SetSkeletonSet(UAbstractSkeletonSetBinding* InSkeletonSetBinding, FName InSkeletonSetName)
+		{
+			SkeletonSetBinding = InSkeletonSetBinding;
+			SkeletonSetName = InSkeletonSetName;
+		}
+
 		FAnimNextGraphInstance& GraphInstance;
+
 		const FReferencePose& RefPose;
+
+		UAbstractSkeletonSetBinding* SkeletonSetBinding = nullptr;
+		FName SkeletonSetName;
+
 		int32 GraphLODLevel = 0;
+
 		TWeakObjectPtr<const USkeletalMeshComponent> BindingObject = nullptr;
 
-		friend UAFANIMGRAPH_API FEvaluationProgram EvaluateGraph(const FEvaluateGraphContext& EvaluateGraphContext, const FWeakTraitPtr& GraphRootPtr);
+		friend UE_API FEvaluationProgram EvaluateGraph(const FEvaluateGraphContext& EvaluateGraphContext, const FWeakTraitPtr& GraphRootPtr);
 		friend struct ::FRigUnit_AnimNextRunAnimationGraph_v1;
 		friend struct ::FRigUnit_AnimNextRunAnimationGraph_v2;
 		friend struct ::FAnimNode_AnimNextGraph;
@@ -75,7 +92,7 @@ namespace UE::UAF
 
 		FEvaluationProgram& EvaluationProgram;
 
-		friend UAFANIMGRAPH_API FEvaluationProgram EvaluateGraph(const FEvaluateGraphContext& EvaluateGraphContext, const FWeakTraitPtr& GraphRootPtr);
+		friend UE_API FEvaluationProgram EvaluateGraph(const FEvaluateGraphContext& EvaluateGraphContext, const FWeakTraitPtr& GraphRootPtr);
 	};
 
 	/**
@@ -142,7 +159,7 @@ namespace UE::UAF
 	 *
 	 * @see IEvaluate::PreEvaluate, IEvaluate::PostEvaluate, IHierarchy::GetChildren
 	 */
-	[[nodiscard]] UAFANIMGRAPH_API FEvaluationProgram EvaluateGraph(const FEvaluateGraphContext& EvaluateGraphContext);
+	[[nodiscard]] UE_API FEvaluationProgram EvaluateGraph(const FEvaluateGraphContext& EvaluateGraphContext);
 
 	/**
 	 * Evaluates a sub-graph starting at its root and produces an evaluation program.
@@ -157,12 +174,15 @@ namespace UE::UAF
 	 *
 	 * @see IEvaluate::PreEvaluate, IEvaluate::PostEvaluate, IHierarchy::GetChildren
 	 */
-	[[nodiscard]] UAFANIMGRAPH_API FEvaluationProgram EvaluateGraph(const FEvaluateGraphContext& EvaluateGraphContext, const FWeakTraitPtr& GraphRootPtr);
+	[[nodiscard]] UE_API FEvaluationProgram EvaluateGraph(const FEvaluateGraphContext& EvaluateGraphContext, const FWeakTraitPtr& GraphRootPtr);
 
 	/**
 	 * Evaluates a sub-graph starting at its root and produces an output pose
 	 */
-	UAFANIMGRAPH_API void EvaluateGraph(const FEvaluateGraphContext& EvaluateGraphContext, FAnimNextGraphLODPose& OutputPose);
+	UE_API void EvaluateGraph(const FEvaluateGraphContext& EvaluateGraphContext, FAnimNextGraphLODPose& OutputPose);
+
+	// Returns whether or not we should use the new abstract skeleton runtime
+	UE_API bool UseAbstractSkeletonRuntime();
 }
 
 #undef UE_API
diff --git a/Engine/Plugins/Experimental/UAF/UAFPoseSearch/Source/UAFPoseSearch/Private/HistoryCollectorTrait.cpp b/Engine/Plugins/Experimental/UAF/UAFPoseSearch/Source/UAFPoseSearch/Private/HistoryCollectorTrait.cpp
--- a/Engine/Plugins/Experimental/UAF/UAFPoseSearch/Source/UAFPoseSearch/Private/HistoryCollectorTrait.cpp
+++ b/Engine/Plugins/Experimental/UAF/UAFPoseSearch/Source/UAFPoseSearch/Private/HistoryCollectorTrait.cpp
@@ -85,18 +85,33 @@ void FAnimNextHistoryCollectorPreEvaluateTask::Execute(UE::UAF::FEvaluationVM& V
 	using namespace UE::UAF;
 	using namespace UE::PoseSearch;
 
-	if (EnumHasAnyFlags(VM.GetFlags(), EEvaluationFlags::Bones))
+	if (VM.GetContext().IsValid())
 	{
 		FPoseHistoryEvaluationHelper PoseHistoryEvalHelper;
 		PoseHistoryEvalHelper.PoseHistoryPtr = InstanceData->PoseHistoryPtr;
 		VM.PushValue(POSEHISTORY_STACK_NAME, MakeUnique<FPoseHistoryEvaluationHelper>(PoseHistoryEvalHelper));
 	}
+	else
+	{
+		if (EnumHasAnyFlags(VM.GetFlags(), EEvaluationFlags::Bones))
+		{
+			FPoseHistoryEvaluationHelper PoseHistoryEvalHelper;
+			PoseHistoryEvalHelper.PoseHistoryPtr = InstanceData->PoseHistoryPtr;
+			VM.PushValue(POSEHISTORY_STACK_NAME, MakeUnique<FPoseHistoryEvaluationHelper>(PoseHistoryEvalHelper));
+		}
+	}
 }
 
 void FAnimNextHistoryCollectorTask::Execute(UE::UAF::FEvaluationVM& VM) const
 {
 	using namespace UE::UAF;
 
+	if (VM.GetContext().IsValid())
+	{
+		// TODO: Implement with new attribute runtime
+		return;
+	}
+
 	if (EnumHasAnyFlags(VM.GetFlags(), EEvaluationFlags::Bones))
 	{
 		if (TUniquePtr<FKeyframeState>* Keyframe = VM.PeekValueMutable<TUniquePtr<FKeyframeState>>(KEYFRAME_STACK_NAME, 0))
diff --git a/Engine/Plugins/Experimental/UAF/UAFTests/Source/UAFTestSuite/Private/UAF/Attributes/AttributeCollectionTest.cpp b/Engine/Plugins/Experimental/UAF/UAFTests/Source/UAFTestSuite/Private/UAF/Attributes/AttributeCollectionTest.cpp
new file mode 100644
--- /dev/null
+++ b/Engine/Plugins/Experimental/UAF/UAFTests/Source/UAFTestSuite/Private/UAF/Attributes/AttributeCollectionTest.cpp
@@ -0,0 +1,365 @@
+// Copyright Epic Games, Inc. All Rights Reserved.
+
+#include "CoreMinimal.h"
+#include "Misc/AutomationTest.h"
+
+#if WITH_DEV_AUTOMATION_TESTS
+
+#include "AnimNextTest.h"
+#include "Animation/BuiltInAttributeTypes.h"
+#include "UAF/AbstractSkeleton/SetBindingFactory.h"
+#include "UAF/AbstractSkeleton/SetCollectionFactory.h"
+#include "UAF/AbstractSkeleton/AbstractSkeletonSetBinding.h"
+#include "UAF/AbstractSkeleton/AbstractSkeletonSetCollection.h"
+#include "UAF/Attributes/AttributeBindingData.h"
+#include "UAF/Attributes/AttributeCollection.h"
+#include "UAF/Attributes/AttributeSetMap.h"
+#include "UAF/Attributes/EngineAttributes.h"
+#include "UObject/UObjectGlobals.h"
+#include "UObject/Package.h"
+
+namespace UE::UAF::Tests
+{
+	template <typename IteratorType, typename ValueType, typename ProjectionType>
+	constexpr bool ContainsBy(IteratorType&& It, const ValueType& Value, ProjectionType Proj)
+	{
+		while (It)
+		{
+			if (Invoke(Proj, It) == Value)
+			{
+				return true;
+			}
+
+			++It;
+		}
+
+		return false;
+	}
+
+	template <typename IteratorType>
+	constexpr int32 IteratorSize(IteratorType&& It)
+	{
+		int32 Size = 0;
+
+		while (It)
+		{
+			++Size;
+			++It;
+		}
+
+		return Size;
+	}
+
+	IMPLEMENT_SIMPLE_AUTOMATION_TEST(FAttributeCollectionTest, "Animation.UAF.Attributes.AttributeCollection", EAutomationTestFlags::EditorContext | EAutomationTestFlags::EngineFilter)
+
+	bool FAttributeCollectionTest::RunTest(const FString& InParameters)
+	{
+		ON_SCOPE_EXIT{ FUtils::CleanupAfterTests(); };
+
+		FAttributeBindingDataPtr BindingData;
+		{
+			UFactory* BindingFactory = NewObject<UFactory>(GetTransientPackage(), UAbstractSkeletonSetBindingFactory::StaticClass());
+			UE_RETURN_ON_ERROR(BindingFactory != nullptr, "Failed to create factory");
+
+			UAbstractSkeletonSetBinding* Binding = CastChecked<UAbstractSkeletonSetBinding>(BindingFactory->FactoryCreateNew(UAbstractSkeletonSetBinding::StaticClass(), GetTransientPackage(), TEXT("TestBindingAsset"), RF_Transient, nullptr, nullptr, NAME_None));
+			UE_RETURN_ON_ERROR(Binding != nullptr, "Failed to create binding asset");
+
+			USkeleton* Skeleton = NewObject<USkeleton>(GetTransientPackage(), USkeleton::StaticClass(), TEXT("TestSkeletonAsset"), RF_Transient);
+			UE_RETURN_ON_ERROR(Skeleton != nullptr, "Failed to create skeleton asset");
+
+			{
+				// Not terribly clean, we cast away the 'const' to modify the skeleton
+				FReferenceSkeleton& RefSkeleton = const_cast<FReferenceSkeleton&>(Skeleton->GetReferenceSkeleton());
+				FReferenceSkeletonModifier SkeletonModifier(RefSkeleton, Skeleton);
+
+				SkeletonModifier.Add(FMeshBoneInfo(TEXT("Root"), TEXT("Root"), INDEX_NONE), FTransform::Identity);
+				SkeletonModifier.Add(FMeshBoneInfo(TEXT("Pelvis"), TEXT("Root"), 0), FTransform::Identity);
+				SkeletonModifier.Add(FMeshBoneInfo(TEXT("Spine_01"), TEXT("Spine_01"), 1), FTransform::Identity);
+				SkeletonModifier.Add(FMeshBoneInfo(TEXT("Spine_02"), TEXT("Spine_02"), 2), FTransform::Identity);
+				SkeletonModifier.Add(FMeshBoneInfo(TEXT("Head"), TEXT("Head"), 3), FTransform::Identity);
+				SkeletonModifier.Add(FMeshBoneInfo(TEXT("LeftArm"), TEXT("LeftArm"), 3), FTransform::Identity);
+				SkeletonModifier.Add(FMeshBoneInfo(TEXT("RightArm"), TEXT("RightArm"), 3), FTransform::Identity);
+
+				// When our modifier is destroyed here, it will rebuild the skeleton
+			}
+
+			UE_RETURN_ON_ERROR(Skeleton->GetReferenceSkeleton().GetNum() == 7, "Failed to build skeleton asset");
+
+			bool bSuccess = Binding->SetSkeleton(Skeleton);
+			UE_RETURN_ON_ERROR(bSuccess, "Failed to set skeleton on our binding");
+
+			UFactory* SecCollectionFactory = NewObject<UFactory>(GetTransientPackage(), UAbstractSkeletonSetCollectionFactory::StaticClass());
+			UE_RETURN_ON_ERROR(SecCollectionFactory != nullptr, "Failed to create factory");
+
+			UAbstractSkeletonSetCollection* SetCollection = CastChecked<UAbstractSkeletonSetCollection>(SecCollectionFactory->FactoryCreateNew(UAbstractSkeletonSetCollection::StaticClass(), GetTransientPackage(), TEXT("TestSetCollectionAsset"), RF_Transient, nullptr, nullptr, NAME_None));
+			UE_RETURN_ON_ERROR(SetCollection != nullptr, "Failed to create set collection asset");
+
+			SetCollection->AddSet(TEXT("FullBody"), NAME_None);
+			SetCollection->AddSet(TEXT("UpperBody"), TEXT("FullBody"));
+
+			bSuccess = Binding->SetSetCollection(SetCollection);
+			UE_RETURN_ON_ERROR(bSuccess, "Failed to set set collection on our binding");
+
+			Binding->AddBoneToSet(TEXT("Root"), TEXT("FullBody"));
+			Binding->AddBoneToSet(TEXT("Pelvis"), TEXT("FullBody"));
+			Binding->AddBoneToSet(TEXT("Spine_01"), TEXT("FullBody"));
+			Binding->AddBoneToSet(TEXT("Spine_02"), TEXT("FullBody"));
+			Binding->AddAttributeToSet(FAnimationAttributeIdentifier(TEXT("AttrFloat0"), INDEX_NONE, NAME_None, FFloatAnimationAttribute::StaticStruct()), TEXT("FullBody"));
+
+			Binding->AddBoneToSet(TEXT("Head"), TEXT("UpperBody"));
+			Binding->AddBoneToSet(TEXT("LeftArm"), TEXT("UpperBody"));
+			Binding->AddBoneToSet(TEXT("RightArm"), TEXT("UpperBody"));
+			Binding->AddAttributeToSet(FAnimationAttributeIdentifier(TEXT("AttrFloat1"), INDEX_NONE, NAME_None, FFloatAnimationAttribute::StaticStruct()), TEXT("UpperBody"));
+
+			BindingData = MakeAttributeBindingData(*Binding);
+			UE_RETURN_ON_ERROR(BindingData.IsValid(), TEXT("Attribute binding data should be valid"));
+		}
+
+		{
+			FAttributeNamedSetPtr FullBodySet = BindingData->FindNamedSet(TEXT("FullBody"));
+			UE_RETURN_ON_ERROR(FullBodySet.IsValid(), TEXT("FullBody named set should be valid"));
+
+			FAttributeCollectionHeap Collection(FullBodySet);
+
+			AddErrorIfFalse(Collection.IsEmpty(), TEXT("Attribute collection should be empty"));
+			AddErrorIfFalse(Collection.GetNamedSet() == FullBodySet, TEXT("Unexpected named set"));
+
+			FAttributeMappingKey FloatToFloatKey = FAttributeMappingKey::MakeFromTo<FFloatAnimationAttribute, FFloatAnimationAttribute>();
+			TAttributeSetMap<FFloatAnimationAttribute>* FloatToFloatMap = Collection.AddSetMap<FFloatAnimationAttribute>(FloatToFloatKey);
+			UE_RETURN_ON_ERROR(FloatToFloatMap != nullptr, TEXT("Float->Float map should be valid"));
+
+			FAttributeMappingKey FloatToIntegerKey = FAttributeMappingKey::MakeFromTo<FFloatAnimationAttribute, FIntegerAnimationAttribute>();
+			TAttributeSetMap<FIntegerAnimationAttribute>* FloatToIntegerMap = Collection.AddSetMap<FIntegerAnimationAttribute>(FloatToIntegerKey);
+			UE_RETURN_ON_ERROR(FloatToIntegerMap != nullptr, TEXT("Float->Integer map should be valid"));
+
+			FAttributeMappingKey FloatToStringKey = FAttributeMappingKey::MakeFromTo<FFloatAnimationAttribute, FStringAnimationAttribute>();
+			TAttributeSetMap<FStringAnimationAttribute>* FloatToStringMap = Collection.AddSetMap<FStringAnimationAttribute>(FloatToStringKey);
+			UE_RETURN_ON_ERROR(FloatToStringMap != nullptr, TEXT("Float->String map should be valid"));
+
+			FAttributeMappingKey TransformToTransformKey = FAttributeMappingKey::MakeFromTo<FBoneTransformAnimationAttribute, FBoneTransformAnimationAttribute>();
+			TAttributeSetMap<FBoneTransformAnimationAttribute>* TransformToTransformMap = Collection.AddSetMap<FBoneTransformAnimationAttribute>(TransformToTransformKey);
+			UE_RETURN_ON_ERROR(TransformToTransformMap != nullptr, TEXT("Transform->Transform map should be valid"));
+
+			FAttributeMappingKey TransformToFloatKey = FAttributeMappingKey::MakeFromTo<FBoneTransformAnimationAttribute, FFloatAnimationAttribute>();
+			TAttributeSetMap<FFloatAnimationAttribute>* TransformToFloatMap = Collection.AddSetMap<FFloatAnimationAttribute>(TransformToFloatKey);
+			UE_RETURN_ON_ERROR(TransformToFloatMap != nullptr, TEXT("Transform->Bone map should be valid"));
+
+			AddErrorIfFalse(!Collection.IsEmpty(), TEXT("Attribute collection should not be empty"));
+			AddErrorIfFalse(Collection.NumSetMaps() == 5, TEXT("Unexpected set map count"));
+
+			AddErrorIfFalse(FloatToFloatMap == Collection.FindSetMap<FFloatAnimationAttribute>(FloatToFloatKey), TEXT("Unexpected result from FindSetMap"));
+			AddErrorIfFalse(FloatToIntegerMap == Collection.FindSetMap<FIntegerAnimationAttribute>(FloatToIntegerKey), TEXT("Unexpected result from FindSetMap"));
+			AddErrorIfFalse(FloatToStringMap == Collection.FindSetMap<FStringAnimationAttribute>(FloatToStringKey), TEXT("Unexpected result from FindSetMap"));
+			AddErrorIfFalse(TransformToTransformMap == Collection.FindSetMap<FBoneTransformAnimationAttribute>(TransformToTransformKey), TEXT("Unexpected result from FindSetMap"));
+			AddErrorIfFalse(TransformToFloatMap == Collection.FindSetMap<FFloatAnimationAttribute>(TransformToFloatKey), TEXT("Unexpected result from FindSetMap"));
+
+			AddErrorIfFalse(Collection.RemoveSetMap(FloatToFloatMap), TEXT("Failed to remove set map"));
+			AddErrorIfFalse(Collection.RemoveSetMap(FloatToIntegerKey), TEXT("Failed to remove set map"));
+			AddErrorIfFalse(Collection.RemoveSetMap(FloatToStringMap), TEXT("Failed to remove set map"));
+			AddErrorIfFalse(Collection.RemoveSetMap(TransformToTransformKey), TEXT("Failed to remove set map"));
+			AddErrorIfFalse(Collection.RemoveSetMap(TransformToFloatMap), TEXT("Failed to remove set map"));
+
+			AddErrorIfFalse(nullptr == Collection.FindSetMap<FFloatAnimationAttribute>(FloatToFloatKey), TEXT("Unexpected result from FindSetMap"));
+			AddErrorIfFalse(nullptr == Collection.FindSetMap<FIntegerAnimationAttribute>(FloatToIntegerKey), TEXT("Unexpected result from FindSetMap"));
+			AddErrorIfFalse(nullptr == Collection.FindSetMap<FStringAnimationAttribute>(FloatToStringKey), TEXT("Unexpected result from FindSetMap"));
+			AddErrorIfFalse(nullptr == Collection.FindSetMap<FBoneTransformAnimationAttribute>(TransformToTransformKey), TEXT("Unexpected result from FindSetMap"));
+			AddErrorIfFalse(nullptr == Collection.FindSetMap<FFloatAnimationAttribute>(TransformToFloatKey), TEXT("Unexpected result from FindSetMap"));
+
+			AddErrorIfFalse(Collection.IsEmpty(), TEXT("Attribute collection should be empty"));
+			AddErrorIfFalse(Collection.NumSetMaps() == 0, TEXT("Unexpected set map count"));
+
+			FloatToFloatMap = Collection.FindOrAddSetMap<FFloatAnimationAttribute>(FloatToFloatKey);
+			UE_RETURN_ON_ERROR(FloatToFloatMap != nullptr, TEXT("Float->Float map should be valid"));
+
+			FloatToIntegerMap = Collection.FindOrAddSetMap<FIntegerAnimationAttribute>(FloatToIntegerKey);
+			UE_RETURN_ON_ERROR(FloatToIntegerMap != nullptr, TEXT("Float->Integer map should be valid"));
+
+			FloatToStringMap = Collection.FindOrAddSetMap<FStringAnimationAttribute>(FloatToStringKey);
+			UE_RETURN_ON_ERROR(FloatToStringMap != nullptr, TEXT("Float->String map should be valid"));
+
+			TransformToTransformMap = Collection.FindOrAddSetMap<FBoneTransformAnimationAttribute>(TransformToTransformKey);
+			UE_RETURN_ON_ERROR(TransformToTransformMap != nullptr, TEXT("Transform->Transform map should be valid"));
+
+			TransformToFloatMap = Collection.FindOrAddSetMap<FFloatAnimationAttribute>(TransformToFloatKey);
+			UE_RETURN_ON_ERROR(TransformToFloatMap != nullptr, TEXT("Transform->Bone map should be valid"));
+
+			AddErrorIfFalse(!Collection.IsEmpty(), TEXT("Attribute collection should not be empty"));
+			AddErrorIfFalse(Collection.NumSetMaps() == 5, TEXT("Unexpected set map count"));
+
+			AddErrorIfFalse(FloatToFloatMap == Collection.FindOrAddSetMap<FFloatAnimationAttribute>(FloatToFloatKey), TEXT("Unexpected result from FindOrAddSetMap"));
+			AddErrorIfFalse(FloatToIntegerMap == Collection.FindOrAddSetMap<FIntegerAnimationAttribute>(FloatToIntegerKey), TEXT("Unexpected result from FindOrAddSetMap"));
+			AddErrorIfFalse(FloatToStringMap == Collection.FindOrAddSetMap<FStringAnimationAttribute>(FloatToStringKey), TEXT("Unexpected result from FindOrAddSetMap"));
+			AddErrorIfFalse(TransformToTransformMap == Collection.FindOrAddSetMap<FBoneTransformAnimationAttribute>(TransformToTransformKey), TEXT("Unexpected result from FindOrAddSetMap"));
+			AddErrorIfFalse(TransformToFloatMap == Collection.FindOrAddSetMap<FFloatAnimationAttribute>(TransformToFloatKey), TEXT("Unexpected result from FindOrAddSetMap"));
+
+			AddErrorIfFalse(!Collection.IsEmpty(), TEXT("Attribute collection should not be empty"));
+			AddErrorIfFalse(Collection.NumSetMaps() == 5, TEXT("Unexpected set map count"));
+
+			// Copy constructor
+			{
+				FAttributeCollectionHeap CollectionCopy(Collection);
+
+				AddErrorIfFalse(CollectionCopy.IsEmpty() == Collection.IsEmpty(), TEXT("Attribute collection should not be empty"));
+				AddErrorIfFalse(CollectionCopy.NumSetMaps() == Collection.NumSetMaps(), TEXT("Unexpected set map count"));
+
+				TAttributeSetMap<FFloatAnimationAttribute>* FloatToFloatMapCopy = CollectionCopy.FindSetMap<FFloatAnimationAttribute>(FloatToFloatKey);
+				UE_RETURN_ON_ERROR(FloatToFloatMapCopy != nullptr, TEXT("Float->Float map should be valid"));
+
+				TAttributeSetMap<FIntegerAnimationAttribute>* FloatToIntegerMapCopy = CollectionCopy.FindSetMap<FIntegerAnimationAttribute>(FloatToIntegerKey);
+				UE_RETURN_ON_ERROR(FloatToIntegerMapCopy != nullptr, TEXT("Float->Integer map should be valid"));
+
+				TAttributeSetMap<FStringAnimationAttribute>* FloatToStringMapCopy = CollectionCopy.FindSetMap<FStringAnimationAttribute>(FloatToStringKey);
+				UE_RETURN_ON_ERROR(FloatToStringMapCopy != nullptr, TEXT("Float->String map should be valid"));
+
+				TAttributeSetMap<FBoneTransformAnimationAttribute>* TransformToTransformMapCopy = CollectionCopy.FindSetMap<FBoneTransformAnimationAttribute>(TransformToTransformKey);
+				UE_RETURN_ON_ERROR(TransformToTransformMapCopy != nullptr, TEXT("Transform->Transform map should be valid"));
+
+				TAttributeSetMap<FFloatAnimationAttribute>* TransformToFloatMapCopy = CollectionCopy.FindSetMap<FFloatAnimationAttribute>(TransformToFloatKey);
+				UE_RETURN_ON_ERROR(TransformToFloatMapCopy != nullptr, TEXT("Transform->Bone map should be valid"));
+
+				AddErrorIfFalse(FloatToFloatMap != FloatToFloatMapCopy, TEXT("Set map instances should be duplicated"));
+				AddErrorIfFalse(FloatToIntegerMap != FloatToIntegerMapCopy, TEXT("Set map instances should be duplicated"));
+				AddErrorIfFalse(FloatToStringMap != FloatToStringMapCopy, TEXT("Set map instances should be duplicated"));
+				AddErrorIfFalse(TransformToTransformMap != TransformToTransformMapCopy, TEXT("Set map instances should be duplicated"));
+				AddErrorIfFalse(TransformToFloatMap != TransformToFloatMapCopy, TEXT("Set map instances should be duplicated"));
+
+				AddErrorIfFalse(FloatToFloatMap->GetTypedSet() == FloatToFloatMapCopy->GetTypedSet(), TEXT("Typed set mismatched"));
+				AddErrorIfFalse(FloatToIntegerMap->GetTypedSet() == FloatToIntegerMapCopy->GetTypedSet(), TEXT("Typed set mismatched"));
+				AddErrorIfFalse(FloatToStringMap->GetTypedSet() == FloatToStringMapCopy->GetTypedSet(), TEXT("Typed set mismatched"));
+				AddErrorIfFalse(TransformToTransformMap->GetTypedSet() == TransformToTransformMapCopy->GetTypedSet(), TEXT("Typed set mismatched"));
+				AddErrorIfFalse(TransformToFloatMap->GetTypedSet() == TransformToFloatMapCopy->GetTypedSet(), TEXT("Typed set mismatched"));
+
+				// Move constructor
+				{
+					FAttributeCollectionHeap CollectionMoved(MoveTemp(CollectionCopy));
+
+					AddErrorIfFalse(CollectionCopy.IsEmpty(), TEXT("Attribute collection should be empty"));
+					AddErrorIfFalse(!CollectionMoved.IsEmpty(), TEXT("Attribute collection should not be empty"));
+					AddErrorIfFalse(CollectionMoved.NumSetMaps() != CollectionCopy.NumSetMaps(), TEXT("Unexpected set map count"));
+
+					TAttributeSetMap<FFloatAnimationAttribute>* FloatToFloatMapMoved = CollectionMoved.FindSetMap<FFloatAnimationAttribute>(FloatToFloatKey);
+					UE_RETURN_ON_ERROR(FloatToFloatMapMoved != nullptr, TEXT("Float->Float map should be valid"));
+
+					TAttributeSetMap<FIntegerAnimationAttribute>* FloatToIntegerMapMoved = CollectionMoved.FindSetMap<FIntegerAnimationAttribute>(FloatToIntegerKey);
+					UE_RETURN_ON_ERROR(FloatToIntegerMapMoved != nullptr, TEXT("Float->Integer map should be valid"));
+
+					TAttributeSetMap<FStringAnimationAttribute>* FloatToStringMapMoved = CollectionMoved.FindSetMap<FStringAnimationAttribute>(FloatToStringKey);
+					UE_RETURN_ON_ERROR(FloatToStringMapMoved != nullptr, TEXT("Float->String map should be valid"));
+
+					TAttributeSetMap<FBoneTransformAnimationAttribute>* TransformToTransformMapMoved = CollectionMoved.FindSetMap<FBoneTransformAnimationAttribute>(TransformToTransformKey);
+					UE_RETURN_ON_ERROR(TransformToTransformMapMoved != nullptr, TEXT("Transform->Transform map should be valid"));
+
+					TAttributeSetMap<FFloatAnimationAttribute>* TransformToFloatMapMoved = CollectionMoved.FindSetMap<FFloatAnimationAttribute>(TransformToFloatKey);
+					UE_RETURN_ON_ERROR(TransformToFloatMapMoved != nullptr, TEXT("Transform->Bone map should be valid"));
+
+					AddErrorIfFalse(FloatToFloatMapMoved == FloatToFloatMapCopy, TEXT("Set map instances should be moved"));
+					AddErrorIfFalse(FloatToIntegerMapMoved == FloatToIntegerMapCopy, TEXT("Set map instances should be moved"));
+					AddErrorIfFalse(FloatToStringMapMoved == FloatToStringMapCopy, TEXT("Set map instances should be moved"));
+					AddErrorIfFalse(TransformToTransformMapMoved == TransformToTransformMapCopy, TEXT("Set map instances should be moved"));
+					AddErrorIfFalse(TransformToFloatMapMoved == TransformToFloatMapCopy, TEXT("Set map instances should be moved"));
+				}
+			}
+
+			// Copy assignment
+			{
+				FAttributeCollectionHeap CollectionCopy;
+				CollectionCopy = Collection;
+
+				AddErrorIfFalse(CollectionCopy.IsEmpty() == Collection.IsEmpty(), TEXT("Attribute collection should not be empty"));
+				AddErrorIfFalse(CollectionCopy.NumSetMaps() == Collection.NumSetMaps(), TEXT("Unexpected set map count"));
+
+				TAttributeSetMap<FFloatAnimationAttribute>* FloatToFloatMapCopy = CollectionCopy.FindSetMap<FFloatAnimationAttribute>(FloatToFloatKey);
+				UE_RETURN_ON_ERROR(FloatToFloatMapCopy != nullptr, TEXT("Float->Float map should be valid"));
+
+				TAttributeSetMap<FIntegerAnimationAttribute>* FloatToIntegerMapCopy = CollectionCopy.FindSetMap<FIntegerAnimationAttribute>(FloatToIntegerKey);
+				UE_RETURN_ON_ERROR(FloatToIntegerMapCopy != nullptr, TEXT("Float->Integer map should be valid"));
+
+				TAttributeSetMap<FStringAnimationAttribute>* FloatToStringMapCopy = CollectionCopy.FindSetMap<FStringAnimationAttribute>(FloatToStringKey);
+				UE_RETURN_ON_ERROR(FloatToStringMapCopy != nullptr, TEXT("Float->String map should be valid"));
+
+				TAttributeSetMap<FBoneTransformAnimationAttribute>* TransformToTransformMapCopy = CollectionCopy.FindSetMap<FBoneTransformAnimationAttribute>(TransformToTransformKey);
+				UE_RETURN_ON_ERROR(TransformToTransformMapCopy != nullptr, TEXT("Transform->Transform map should be valid"));
+
+				TAttributeSetMap<FFloatAnimationAttribute>* TransformToFloatMapCopy = CollectionCopy.FindSetMap<FFloatAnimationAttribute>(TransformToFloatKey);
+				UE_RETURN_ON_ERROR(TransformToFloatMapCopy != nullptr, TEXT("Transform->Bone map should be valid"));
+
+				AddErrorIfFalse(FloatToFloatMap != FloatToFloatMapCopy, TEXT("Set map instances should be duplicated"));
+				AddErrorIfFalse(FloatToIntegerMap != FloatToIntegerMapCopy, TEXT("Set map instances should be duplicated"));
+				AddErrorIfFalse(FloatToStringMap != FloatToStringMapCopy, TEXT("Set map instances should be duplicated"));
+				AddErrorIfFalse(TransformToTransformMap != TransformToTransformMapCopy, TEXT("Set map instances should be duplicated"));
+				AddErrorIfFalse(TransformToFloatMap != TransformToFloatMapCopy, TEXT("Set map instances should be duplicated"));
+
+				AddErrorIfFalse(FloatToFloatMap->GetTypedSet() == FloatToFloatMapCopy->GetTypedSet(), TEXT("Typed set mismatched"));
+				AddErrorIfFalse(FloatToIntegerMap->GetTypedSet() == FloatToIntegerMapCopy->GetTypedSet(), TEXT("Typed set mismatched"));
+				AddErrorIfFalse(FloatToStringMap->GetTypedSet() == FloatToStringMapCopy->GetTypedSet(), TEXT("Typed set mismatched"));
+				AddErrorIfFalse(TransformToTransformMap->GetTypedSet() == TransformToTransformMapCopy->GetTypedSet(), TEXT("Typed set mismatched"));
+				AddErrorIfFalse(TransformToFloatMap->GetTypedSet() == TransformToFloatMapCopy->GetTypedSet(), TEXT("Typed set mismatched"));
+
+				// Move assignment
+				{
+					FAttributeCollectionHeap CollectionMoved;
+					CollectionMoved = MoveTemp(CollectionCopy);
+
+					AddErrorIfFalse(CollectionCopy.IsEmpty(), TEXT("Attribute collection should be empty"));
+					AddErrorIfFalse(!CollectionMoved.IsEmpty(), TEXT("Attribute collection should not be empty"));
+					AddErrorIfFalse(CollectionMoved.NumSetMaps() != CollectionCopy.NumSetMaps(), TEXT("Unexpected set map count"));
+
+					TAttributeSetMap<FFloatAnimationAttribute>* FloatToFloatMapMoved = CollectionMoved.FindSetMap<FFloatAnimationAttribute>(FloatToFloatKey);
+					UE_RETURN_ON_ERROR(FloatToFloatMapMoved != nullptr, TEXT("Float->Float map should be valid"));
+
+					TAttributeSetMap<FIntegerAnimationAttribute>* FloatToIntegerMapMoved = CollectionMoved.FindSetMap<FIntegerAnimationAttribute>(FloatToIntegerKey);
+					UE_RETURN_ON_ERROR(FloatToIntegerMapMoved != nullptr, TEXT("Float->Integer map should be valid"));
+
+					TAttributeSetMap<FStringAnimationAttribute>* FloatToStringMapMoved = CollectionMoved.FindSetMap<FStringAnimationAttribute>(FloatToStringKey);
+					UE_RETURN_ON_ERROR(FloatToStringMapMoved != nullptr, TEXT("Float->String map should be valid"));
+
+					TAttributeSetMap<FBoneTransformAnimationAttribute>* TransformToTransformMapMoved = CollectionMoved.FindSetMap<FBoneTransformAnimationAttribute>(TransformToTransformKey);
+					UE_RETURN_ON_ERROR(TransformToTransformMapMoved != nullptr, TEXT("Transform->Transform map should be valid"));
+
+					TAttributeSetMap<FFloatAnimationAttribute>* TransformToFloatMapMoved = CollectionMoved.FindSetMap<FFloatAnimationAttribute>(TransformToFloatKey);
+					UE_RETURN_ON_ERROR(TransformToFloatMapMoved != nullptr, TEXT("Transform->Bone map should be valid"));
+
+					AddErrorIfFalse(FloatToFloatMapMoved == FloatToFloatMapCopy, TEXT("Set map instances should be moved"));
+					AddErrorIfFalse(FloatToIntegerMapMoved == FloatToIntegerMapCopy, TEXT("Set map instances should be moved"));
+					AddErrorIfFalse(FloatToStringMapMoved == FloatToStringMapCopy, TEXT("Set map instances should be moved"));
+					AddErrorIfFalse(TransformToTransformMapMoved == TransformToTransformMapCopy, TEXT("Set map instances should be moved"));
+					AddErrorIfFalse(TransformToFloatMapMoved == TransformToFloatMapCopy, TEXT("Set map instances should be moved"));
+				}
+			}
+
+			// Iterators
+			{
+				AddErrorIfFalse(IteratorSize(Collection.CreateSetMapIterator()) == 5, TEXT("Iterator did not have expected size"));
+				AddErrorIfFalse(IteratorSize(Collection.CreateConstSetMapIterator()) == 5, TEXT("Iterator did not have expected size"));
+
+				AddErrorIfFalse(ContainsBy(Collection.CreateSetMapIterator(), FloatToFloatMap->GetTypedSet()->GetType(), [](const typename FAttributeCollection::TSetMapIterator& It) { return It.GetAttributeType(); }), TEXT("Iterator did not contain expected attribute type"));
+				AddErrorIfFalse(ContainsBy(Collection.CreateSetMapIterator(), FloatToIntegerMap->GetTypedSet()->GetType(), [](const typename FAttributeCollection::TSetMapIterator& It) { return It.GetAttributeType(); }), TEXT("Iterator did not contain expected attribute type"));
+				AddErrorIfFalse(ContainsBy(Collection.CreateSetMapIterator(), FloatToStringMap->GetTypedSet()->GetType(), [](const typename FAttributeCollection::TSetMapIterator& It) { return It.GetAttributeType(); }), TEXT("Iterator did not contain expected attribute type"));
+				AddErrorIfFalse(ContainsBy(Collection.CreateSetMapIterator(), TransformToTransformMap->GetTypedSet()->GetType(), [](const typename FAttributeCollection::TSetMapIterator& It) { return It.GetAttributeType(); }), TEXT("Iterator did not contain expected attribute type"));
+				AddErrorIfFalse(ContainsBy(Collection.CreateSetMapIterator(), TransformToFloatMap->GetTypedSet()->GetType(), [](const typename FAttributeCollection::TSetMapIterator& It) { return It.GetAttributeType(); }), TEXT("Iterator did not contain expected attribute type"));
+
+				AddErrorIfFalse(ContainsBy(Collection.CreateSetMapIterator(), FloatToFloatMap, [](const typename FAttributeCollection::TSetMapIterator& It) { return It.GetMap(); }), TEXT("Iterator did not contain expected set map"));
+				AddErrorIfFalse(ContainsBy(Collection.CreateSetMapIterator(), FloatToIntegerMap, [](const typename FAttributeCollection::TSetMapIterator& It) { return It.GetMap(); }), TEXT("Iterator did not contain expected set map"));
+				AddErrorIfFalse(ContainsBy(Collection.CreateSetMapIterator(), FloatToStringMap, [](const typename FAttributeCollection::TSetMapIterator& It) { return It.GetMap(); }), TEXT("Iterator did not contain expected set map"));
+				AddErrorIfFalse(ContainsBy(Collection.CreateSetMapIterator(), TransformToTransformMap, [](const typename FAttributeCollection::TSetMapIterator& It) { return It.GetMap(); }), TEXT("Iterator did not contain expected set map"));
+				AddErrorIfFalse(ContainsBy(Collection.CreateSetMapIterator(), TransformToFloatMap, [](const typename FAttributeCollection::TSetMapIterator& It) { return It.GetMap(); }), TEXT("Iterator did not contain expected set map"));
+
+				AddErrorIfFalse(ContainsBy(Collection.CreateConstSetMapIterator(), FloatToFloatMap->GetTypedSet()->GetType(), [](const typename FAttributeCollection::TConstSetMapIterator& It) { return It.GetAttributeType(); }), TEXT("Iterator did not contain expected attribute type"));
+				AddErrorIfFalse(ContainsBy(Collection.CreateConstSetMapIterator(), FloatToIntegerMap->GetTypedSet()->GetType(), [](const typename FAttributeCollection::TConstSetMapIterator& It) { return It.GetAttributeType(); }), TEXT("Iterator did not contain expected attribute type"));
+				AddErrorIfFalse(ContainsBy(Collection.CreateConstSetMapIterator(), FloatToStringMap->GetTypedSet()->GetType(), [](const typename FAttributeCollection::TConstSetMapIterator& It) { return It.GetAttributeType(); }), TEXT("Iterator did not contain expected attribute type"));
+				AddErrorIfFalse(ContainsBy(Collection.CreateConstSetMapIterator(), TransformToTransformMap->GetTypedSet()->GetType(), [](const typename FAttributeCollection::TConstSetMapIterator& It) { return It.GetAttributeType(); }), TEXT("Iterator did not contain expected attribute type"));
+				AddErrorIfFalse(ContainsBy(Collection.CreateConstSetMapIterator(), TransformToFloatMap->GetTypedSet()->GetType(), [](const typename FAttributeCollection::TConstSetMapIterator& It) { return It.GetAttributeType(); }), TEXT("Iterator did not contain expected attribute type"));
+
+				AddErrorIfFalse(ContainsBy(Collection.CreateConstSetMapIterator(), FloatToFloatMap, [](const typename FAttributeCollection::TConstSetMapIterator& It) { return It.GetMap(); }), TEXT("Iterator did not contain expected set map"));
+				AddErrorIfFalse(ContainsBy(Collection.CreateConstSetMapIterator(), FloatToIntegerMap, [](const typename FAttributeCollection::TConstSetMapIterator& It) { return It.GetMap(); }), TEXT("Iterator did not contain expected set map"));
+				AddErrorIfFalse(ContainsBy(Collection.CreateConstSetMapIterator(), FloatToStringMap, [](const typename FAttributeCollection::TConstSetMapIterator& It) { return It.GetMap(); }), TEXT("Iterator did not contain expected set map"));
+				AddErrorIfFalse(ContainsBy(Collection.CreateConstSetMapIterator(), TransformToTransformMap, [](const typename FAttributeCollection::TConstSetMapIterator& It) { return It.GetMap(); }), TEXT("Iterator did not contain expected set map"));
+				AddErrorIfFalse(ContainsBy(Collection.CreateConstSetMapIterator(), TransformToFloatMap, [](const typename FAttributeCollection::TConstSetMapIterator& It) { return It.GetMap(); }), TEXT("Iterator did not contain expected set map"));
+			}
+		}
+
+		return true;
+	}
+}
+
+#endif	// WITH_DEV_AUTOMATION_TESTS
diff --git a/Engine/Plugins/Experimental/UAF/UAFTests/Source/UAFTestSuite/Private/UAF/Attributes/AttributeSetMapTest.cpp b/Engine/Plugins/Experimental/UAF/UAFTests/Source/UAFTestSuite/Private/UAF/Attributes/AttributeSetMapTest.cpp
new file mode 100644
--- /dev/null
+++ b/Engine/Plugins/Experimental/UAF/UAFTests/Source/UAFTestSuite/Private/UAF/Attributes/AttributeSetMapTest.cpp
@@ -0,0 +1,203 @@
+// Copyright Epic Games, Inc. All Rights Reserved.
+
+#include "CoreMinimal.h"
+#include "Misc/AutomationTest.h"
+
+#if WITH_DEV_AUTOMATION_TESTS
+
+#include "AnimNextTest.h"
+#include "Animation/BuiltInAttributeTypes.h"
+#include "UAF/AbstractSkeleton/SetBindingFactory.h"
+#include "UAF/AbstractSkeleton/SetCollectionFactory.h"
+#include "UAF/AbstractSkeleton/AbstractSkeletonSetBinding.h"
+#include "UAF/AbstractSkeleton/AbstractSkeletonSetCollection.h"
+#include "UAF/Attributes/AttributeBindingData.h"
+#include "UAF/Attributes/AttributeSetMap.h"
+#include "UAF/Attributes/EngineAttributes.h"
+#include "UObject/UObjectGlobals.h"
+#include "UObject/Package.h"
+
+namespace UE::UAF::Tests
+{
+	IMPLEMENT_SIMPLE_AUTOMATION_TEST(FAttributeSetMapTest, "Animation.UAF.Attributes.AttributeSetMap", EAutomationTestFlags::EditorContext | EAutomationTestFlags::EngineFilter)
+
+	bool FAttributeSetMapTest::RunTest(const FString& InParameters)
+	{
+		ON_SCOPE_EXIT{ FUtils::CleanupAfterTests(); };
+
+		FAttributeBindingDataPtr BindingData;
+		{
+			UFactory* BindingFactory = NewObject<UFactory>(GetTransientPackage(), UAbstractSkeletonSetBindingFactory::StaticClass());
+			UE_RETURN_ON_ERROR(BindingFactory != nullptr, "Failed to create factory");
+
+			UAbstractSkeletonSetBinding* Binding = CastChecked<UAbstractSkeletonSetBinding>(BindingFactory->FactoryCreateNew(UAbstractSkeletonSetBinding::StaticClass(), GetTransientPackage(), TEXT("TestBindingAsset"), RF_Transient, nullptr, nullptr, NAME_None));
+			UE_RETURN_ON_ERROR(Binding != nullptr, "Failed to create binding asset");
+
+			USkeleton* Skeleton = NewObject<USkeleton>(GetTransientPackage(), USkeleton::StaticClass(), TEXT("TestSkeletonAsset"), RF_Transient);
+			UE_RETURN_ON_ERROR(Skeleton != nullptr, "Failed to create skeleton asset");
+
+			{
+				// Not terribly clean, we cast away the 'const' to modify the skeleton
+				FReferenceSkeleton& RefSkeleton = const_cast<FReferenceSkeleton&>(Skeleton->GetReferenceSkeleton());
+				FReferenceSkeletonModifier SkeletonModifier(RefSkeleton, Skeleton);
+
+				SkeletonModifier.Add(FMeshBoneInfo(TEXT("Root"), TEXT("Root"), INDEX_NONE), FTransform::Identity);
+				SkeletonModifier.Add(FMeshBoneInfo(TEXT("Pelvis"), TEXT("Root"), 0), FTransform::Identity);
+				SkeletonModifier.Add(FMeshBoneInfo(TEXT("Spine_01"), TEXT("Spine_01"), 1), FTransform::Identity);
+				SkeletonModifier.Add(FMeshBoneInfo(TEXT("Spine_02"), TEXT("Spine_02"), 2), FTransform::Identity);
+				SkeletonModifier.Add(FMeshBoneInfo(TEXT("Head"), TEXT("Head"), 3), FTransform::Identity);
+				SkeletonModifier.Add(FMeshBoneInfo(TEXT("LeftArm"), TEXT("LeftArm"), 3), FTransform::Identity);
+				SkeletonModifier.Add(FMeshBoneInfo(TEXT("RightArm"), TEXT("RightArm"), 3), FTransform::Identity);
+
+				// When our modifier is destroyed here, it will rebuild the skeleton
+			}
+
+			UE_RETURN_ON_ERROR(Skeleton->GetReferenceSkeleton().GetNum() == 7, "Failed to build skeleton asset");
+
+			bool bSuccess = Binding->SetSkeleton(Skeleton);
+			UE_RETURN_ON_ERROR(bSuccess, "Failed to set skeleton on our binding");
+
+			UFactory* SecCollectionFactory = NewObject<UFactory>(GetTransientPackage(), UAbstractSkeletonSetCollectionFactory::StaticClass());
+			UE_RETURN_ON_ERROR(SecCollectionFactory != nullptr, "Failed to create factory");
+
+			UAbstractSkeletonSetCollection* SetCollection = CastChecked<UAbstractSkeletonSetCollection>(SecCollectionFactory->FactoryCreateNew(UAbstractSkeletonSetCollection::StaticClass(), GetTransientPackage(), TEXT("TestSetCollectionAsset"), RF_Transient, nullptr, nullptr, NAME_None));
+			UE_RETURN_ON_ERROR(SetCollection != nullptr, "Failed to create set collection asset");
+
+			SetCollection->AddSet(TEXT("FullBody"), NAME_None);
+			SetCollection->AddSet(TEXT("UpperBody"), TEXT("FullBody"));
+
+			bSuccess = Binding->SetSetCollection(SetCollection);
+			UE_RETURN_ON_ERROR(bSuccess, "Failed to set set collection on our binding");
+
+			Binding->AddBoneToSet(TEXT("Root"), TEXT("FullBody"));
+			Binding->AddBoneToSet(TEXT("Pelvis"), TEXT("FullBody"));
+			Binding->AddBoneToSet(TEXT("Spine_01"), TEXT("FullBody"));
+			Binding->AddBoneToSet(TEXT("Spine_02"), TEXT("FullBody"));
+			Binding->AddAttributeToSet(FAnimationAttributeIdentifier(TEXT("AttrFloat0"), INDEX_NONE, NAME_None, FFloatAnimationAttribute::StaticStruct()), TEXT("FullBody"));
+
+			Binding->AddBoneToSet(TEXT("Head"), TEXT("UpperBody"));
+			Binding->AddBoneToSet(TEXT("LeftArm"), TEXT("UpperBody"));
+			Binding->AddBoneToSet(TEXT("RightArm"), TEXT("UpperBody"));
+			Binding->AddAttributeToSet(FAnimationAttributeIdentifier(TEXT("AttrFloat1"), INDEX_NONE, NAME_None, FFloatAnimationAttribute::StaticStruct()), TEXT("UpperBody"));
+
+			BindingData = MakeAttributeBindingData(*Binding);
+			UE_RETURN_ON_ERROR(BindingData.IsValid(), TEXT("Attribute binding data should be valid"));
+		}
+
+		{
+			FReallocFun Allocator = &FAllocatorTypeTrait<FDefaultAllocator>::Realloc;
+
+			FAttributeNamedSetPtr FullBodySet = BindingData->FindNamedSet(TEXT("FullBody"));
+			UE_RETURN_ON_ERROR(FullBodySet.IsValid(), TEXT("FullBody named set should be valid"));
+
+			FAttributeTypedSetPtr FullBodyBoneSet = FullBodySet->FindTypedSet<FBoneTransformAnimationAttribute>();
+			UE_RETURN_ON_ERROR(FullBodyBoneSet.IsValid(), TEXT("FullBody bone typed set should be valid"));
+
+			{
+				TAttributeSetMap<FFloatAnimationAttribute>& BoneToFloatMap = *MakeAttributeSetMap<FFloatAnimationAttribute>(FAttributeSetMap::FConstructArgs(FullBodyBoneSet, FFloatAnimationAttribute::StaticStruct(), Allocator));
+
+				AddErrorIfFalse(!BoneToFloatMap.IsEmpty(), TEXT("Attribute set map should not be empty"));
+				AddErrorIfFalse(BoneToFloatMap.Num() == FullBodyBoneSet->Num(), TEXT("Unexpected set map size"));
+				AddErrorIfFalse(BoneToFloatMap.GetValueType() == FFloatAnimationAttribute::StaticStruct(), TEXT("Unexpected set map value type"));
+
+				FFloatAnimationAttribute Value0{ 123.44f };
+				FFloatAnimationAttribute Value1{ -123.44f };
+				BoneToFloatMap[FAttributeSetIndex(0)] = Value0;
+				BoneToFloatMap[FAttributeSetIndex(1)] = Value1;
+
+				AddErrorIfFalse(BoneToFloatMap[FAttributeSetIndex(0)].Value == Value0.Value, TEXT("Unexpected mapped value"));
+				AddErrorIfFalse(BoneToFloatMap[FAttributeSetIndex(1)].Value == Value1.Value, TEXT("Unexpected mapped value"));
+
+				FFloatAnimationAttribute Value2{ -0.13544f };
+				BoneToFloatMap.SetValue(FAttributeSetIndex(2), Value2);
+
+				AddErrorIfFalse(BoneToFloatMap.GetValue<FFloatAnimationAttribute>(FAttributeSetIndex(2)).Value == Value2.Value, TEXT("Unexpected mapped value"));
+
+				{
+					TAttributeSetMap<FFloatAnimationAttribute>& BoneToFloatMapCopy = *MakeAttributeSetMap<FFloatAnimationAttribute>(FAttributeSetMap::FConstructArgs(FullBodyBoneSet, FFloatAnimationAttribute::StaticStruct(), Allocator));
+					BoneToFloatMap.CopyTo(BoneToFloatMapCopy);
+
+					AddErrorIfFalse(BoneToFloatMapCopy[FAttributeSetIndex(0)].Value == Value0.Value, TEXT("Unexpected mapped value"));
+					AddErrorIfFalse(BoneToFloatMapCopy[FAttributeSetIndex(1)].Value == Value1.Value, TEXT("Unexpected mapped value"));
+					AddErrorIfFalse(BoneToFloatMapCopy[FAttributeSetIndex(2)].Value == Value2.Value, TEXT("Unexpected mapped value"));
+
+					ReleaseAttributeSetMap(&BoneToFloatMapCopy);
+				}
+
+				ReleaseAttributeSetMap(&BoneToFloatMap);
+			}
+
+			{
+				TAttributeSetMap<FBoneTransformAnimationAttribute>& BoneToTransformMap = *MakeAttributeSetMap<FBoneTransformAnimationAttribute>(FAttributeSetMap::FConstructArgs(FullBodyBoneSet, FBoneTransformAnimationAttribute::StaticStruct(), Allocator));
+
+				AddErrorIfFalse(!BoneToTransformMap.IsEmpty(), TEXT("Attribute set map should not be empty"));
+				AddErrorIfFalse(BoneToTransformMap.Num() == FullBodyBoneSet->Num(), TEXT("Unexpected set map size"));
+				AddErrorIfFalse(BoneToTransformMap.GetValueType() == FBoneTransformAnimationAttribute::StaticStruct(), TEXT("Unexpected set map value type"));
+
+				FBoneTransformAnimationAttribute Value0{ FTransform::Identity };
+				FBoneTransformAnimationAttribute Value1{ FTransform(FVector(1.2, 3.4, 5.6)) };
+				BoneToTransformMap[FAttributeSetIndex(0)] = Value0;
+				BoneToTransformMap[FAttributeSetIndex(1)] = Value1;
+
+				AddErrorIfFalse(BoneToTransformMap[FAttributeSetIndex(0)].Value.Equals(Value0.Value), TEXT("Unexpected mapped value"));
+				AddErrorIfFalse(BoneToTransformMap[FAttributeSetIndex(1)].Value.Equals(Value1.Value), TEXT("Unexpected mapped value"));
+
+				FBoneTransformAnimationAttribute Value2{ FTransform(FVector(0.233, 331.54, 571.146)) };
+				BoneToTransformMap.SetValue(FAttributeSetIndex(2), Value2);
+
+				AddErrorIfFalse(BoneToTransformMap.GetValue<FBoneTransformAnimationAttribute>(FAttributeSetIndex(2)).Value.Equals(Value2.Value), TEXT("Unexpected mapped value"));
+
+				{
+					TAttributeSetMap<FBoneTransformAnimationAttribute>& BoneToTransformMapCopy = *MakeAttributeSetMap<FBoneTransformAnimationAttribute>(FAttributeSetMap::FConstructArgs(FullBodyBoneSet, FBoneTransformAnimationAttribute::StaticStruct(), Allocator));
+					BoneToTransformMap.CopyTo(BoneToTransformMapCopy);
+
+					AddErrorIfFalse(BoneToTransformMapCopy[FAttributeSetIndex(0)].Value.Equals(Value0.Value), TEXT("Unexpected mapped value"));
+					AddErrorIfFalse(BoneToTransformMapCopy[FAttributeSetIndex(1)].Value.Equals(Value1.Value), TEXT("Unexpected mapped value"));
+					AddErrorIfFalse(BoneToTransformMapCopy[FAttributeSetIndex(2)].Value.Equals(Value2.Value), TEXT("Unexpected mapped value"));
+
+					ReleaseAttributeSetMap(&BoneToTransformMapCopy);
+				}
+
+				ReleaseAttributeSetMap(&BoneToTransformMap);
+			}
+
+			{
+				TAttributeSetMap<FStringAnimationAttribute>& BoneToStringMap = *MakeAttributeSetMap<FStringAnimationAttribute>(FAttributeSetMap::FConstructArgs(FullBodyBoneSet, FStringAnimationAttribute::StaticStruct(), Allocator));
+
+				AddErrorIfFalse(!BoneToStringMap.IsEmpty(), TEXT("Attribute set map should not be empty"));
+				AddErrorIfFalse(BoneToStringMap.Num() == FullBodyBoneSet->Num(), TEXT("Unexpected set map size"));
+				AddErrorIfFalse(BoneToStringMap.GetValueType() == FStringAnimationAttribute::StaticStruct(), TEXT("Unexpected set map value type"));
+
+				FStringAnimationAttribute Value0{ TEXT("FooBar") };
+				FStringAnimationAttribute Value1{ TEXT("Sweet123")};
+				BoneToStringMap[FAttributeSetIndex(0)] = Value0;
+				BoneToStringMap[FAttributeSetIndex(1)] = Value1;
+
+				AddErrorIfFalse(BoneToStringMap[FAttributeSetIndex(0)].Value.Equals(Value0.Value), TEXT("Unexpected mapped value"));
+				AddErrorIfFalse(BoneToStringMap[FAttributeSetIndex(1)].Value.Equals(Value1.Value), TEXT("Unexpected mapped value"));
+
+				FStringAnimationAttribute Value2{ TEXT("01Wow")};
+				BoneToStringMap.SetValue(FAttributeSetIndex(2), Value2);
+
+				AddErrorIfFalse(BoneToStringMap.GetValue<FStringAnimationAttribute>(FAttributeSetIndex(2)).Value.Equals(Value2.Value), TEXT("Unexpected mapped value"));
+
+				{
+					TAttributeSetMap<FStringAnimationAttribute>& BoneToStringMapCopy = *MakeAttributeSetMap<FStringAnimationAttribute>(FAttributeSetMap::FConstructArgs(FullBodyBoneSet, FStringAnimationAttribute::StaticStruct(), Allocator));
+					BoneToStringMap.CopyTo(BoneToStringMapCopy);
+
+					AddErrorIfFalse(BoneToStringMapCopy[FAttributeSetIndex(0)].Value.Equals(Value0.Value), TEXT("Unexpected mapped value"));
+					AddErrorIfFalse(BoneToStringMapCopy[FAttributeSetIndex(1)].Value.Equals(Value1.Value), TEXT("Unexpected mapped value"));
+					AddErrorIfFalse(BoneToStringMapCopy[FAttributeSetIndex(2)].Value.Equals(Value2.Value), TEXT("Unexpected mapped value"));
+
+					ReleaseAttributeSetMap(&BoneToStringMapCopy);
+				}
+
+				ReleaseAttributeSetMap(&BoneToStringMap);
+			}
+		}
+
+		return true;
+	}
+}
+
+#endif	// WITH_DEV_AUTOMATION_TESTS
diff --git a/Engine/Plugins/Experimental/UAF/UAFWarping/Source/UAFWarping/Private/StrafeWarpingTrait.cpp b/Engine/Plugins/Experimental/UAF/UAFWarping/Source/UAFWarping/Private/StrafeWarpingTrait.cpp
--- a/Engine/Plugins/Experimental/UAF/UAFWarping/Source/UAFWarping/Private/StrafeWarpingTrait.cpp
+++ b/Engine/Plugins/Experimental/UAF/UAFWarping/Source/UAFWarping/Private/StrafeWarpingTrait.cpp
@@ -185,6 +185,12 @@ void FAnimNextStrafeWarpingTask::Execute(UE::UAF::FEvaluationVM& VM) const
 {
 	using namespace UE::UAF;
 	QUICK_SCOPE_CYCLE_COUNTER(FAnimNextStrafeWarpingTask_Execute);
+
+	if (VM.GetContext().IsValid())
+	{
+		// TODO: Implement with new attribute runtime
+		return;
+	}
 	
 	if (InstanceData->DeltaTime == 0.f)
 	{
diff --git a/Engine/Source/Runtime/Engine/Classes/Animation/BuiltInAttributeTypes.h b/Engine/Source/Runtime/Engine/Classes/Animation/BuiltInAttributeTypes.h
--- a/Engine/Source/Runtime/Engine/Classes/Animation/BuiltInAttributeTypes.h
+++ b/Engine/Source/Runtime/Engine/Classes/Animation/BuiltInAttributeTypes.h
@@ -329,7 +329,7 @@ namespace UE
 
 		/** Non blendable types*/
 		template<>
-		struct TAttributeTypeTraits<FNonBlendableQuaternionAnimationAttribute> : public TAttributeTypeTraitsBase<FNonBlendableQuaternionAnimationAttribute>
+		struct TAttributeTypeTraits<FNonBlendableQuaternionAnimationAttribute> : public TAttributeTypeTraits<FQuaternionAnimationAttribute>
 		{
 			enum
 			{
@@ -339,7 +339,7 @@ namespace UE
 
 		/** Non blendable types*/
 		template<>
-		struct TAttributeTypeTraits<FNonBlendableVectorAnimationAttribute> : public TAttributeTypeTraitsBase<FNonBlendableVectorAnimationAttribute>
+		struct TAttributeTypeTraits<FNonBlendableVectorAnimationAttribute> : public TAttributeTypeTraits<FVectorAnimationAttribute>
 		{
 			enum
 			{
@@ -349,7 +349,7 @@ namespace UE
 		
 		/** Non blendable types*/
 		template<>
-		struct TAttributeTypeTraits<FNonBlendableTransformAnimationAttribute> : public TAttributeTypeTraitsBase<FNonBlendableTransformAnimationAttribute>
+		struct TAttributeTypeTraits<FNonBlendableTransformAnimationAttribute> : public TAttributeTypeTraits<FTransformAnimationAttribute>
 		{
 			enum
 			{
@@ -358,7 +358,7 @@ namespace UE
 		};
 		
 		template<>
-		struct TAttributeTypeTraits<FNonBlendableFloatAnimationAttribute> : public TAttributeTypeTraitsBase<FNonBlendableFloatAnimationAttribute>
+		struct TAttributeTypeTraits<FNonBlendableFloatAnimationAttribute> : public TAttributeTypeTraits<FFloatAnimationAttribute>
 		{
 			enum
 			{
@@ -367,7 +367,7 @@ namespace UE
 		};
 
 		template<>
-		struct TAttributeTypeTraits<FNonBlendableIntegerAnimationAttribute> : public TAttributeTypeTraitsBase<FNonBlendableIntegerAnimationAttribute>
+		struct TAttributeTypeTraits<FNonBlendableIntegerAnimationAttribute> : public TAttributeTypeTraits<FIntegerAnimationAttribute>
 		{
 			enum
 			{
diff --git a/Engine/Source/Runtime/Engine/Classes/Animation/Skeleton.h b/Engine/Source/Runtime/Engine/Classes/Animation/Skeleton.h
--- a/Engine/Source/Runtime/Engine/Classes/Animation/Skeleton.h
+++ b/Engine/Source/Runtime/Engine/Classes/Animation/Skeleton.h
@@ -945,19 +945,30 @@ protected:
 	bool CreateReferenceSkeletonFromMesh(const USkinnedAsset* InSkinnedAsset, const TArray<int32> & RequiredRefBones);
 
 #if WITH_EDITOR
-	DECLARE_MULTICAST_DELEGATE( FOnSkeletonHierarchyChangedMulticaster );
+	DECLARE_TS_MULTICAST_DELEGATE( FOnSkeletonHierarchyChangedMulticaster );
 	FOnSkeletonHierarchyChangedMulticaster OnSkeletonHierarchyChanged;
 
+	DECLARE_TS_MULTICAST_DELEGATE( FOnReferenceSkeletonChangedMulticaster );
+	FOnReferenceSkeletonChangedMulticaster OnReferenceSkeletonChanged;
+
 	/** Call this when the skeleton has changed to fix dependent assets */
 	ENGINE_API void HandleSkeletonHierarchyChange(bool bShowProgress = true);
 
+	/** Call this when the reference skeleton has changed to fix dependent assets */
+	ENGINE_API void HandleReferenceSkeletonChange() const;
+
 public:
 	typedef FOnSkeletonHierarchyChangedMulticaster::FDelegate FOnSkeletonHierarchyChanged;
+	typedef FOnReferenceSkeletonChangedMulticaster::FDelegate FOnReferenceSkeletonChanged;
 
-	/** Registers a delegate to be called after notification has changed*/
+	/** Registers a delegate to be called after the skeleton hierarchy has changed */
 	ENGINE_API void RegisterOnSkeletonHierarchyChanged(const FOnSkeletonHierarchyChanged& Delegate);
 	ENGINE_API void UnregisterOnSkeletonHierarchyChanged(FDelegateUserObject Unregister);
 
+	/** Registers a delegate to be called after the reference skeleton has changed */
+	ENGINE_API void RegisterOnReferenceSkeletonChanged(const FOnSkeletonHierarchyChanged& Delegate);
+	ENGINE_API void UnregisterOnReferenceSkeletonChanged(FDelegateUserObject Unregister);
+
 	/** Removes the supplied bones from the skeleton */
 	ENGINE_API void RemoveBonesFromSkeleton(const TArray<FName>& BonesToRemove, bool bRemoveChildBones);
 
@@ -1012,7 +1023,7 @@ protected:
 	TArray<TObjectPtr<UAssetUserData>> AssetUserDataEditorOnly;
 #endif
 
-
+	friend FReferenceSkeleton;
 	friend struct FReferenceSkeletonModifier;
 	friend class FEditableSkeleton;
 };
diff --git a/Engine/Source/Runtime/Engine/Private/Animation/ReferenceSkeleton.cpp b/Engine/Source/Runtime/Engine/Private/Animation/ReferenceSkeleton.cpp
--- a/Engine/Source/Runtime/Engine/Private/Animation/ReferenceSkeleton.cpp
+++ b/Engine/Source/Runtime/Engine/Private/Animation/ReferenceSkeleton.cpp
@@ -534,6 +534,14 @@ void FReferenceSkeleton::RebuildRefSkeleton(const USkeleton* Skeleton, bool bReb
 
 	// Full rebuild of all compatible with this and with ones we are compatible with.
 	UE::Anim::FSkeletonRemappingRegistry::Get().RefreshMappings(Skeleton);
+
+#if WITH_EDITOR
+	if (Skeleton)
+	{
+		// Notify the skeleton that we changed
+		Skeleton->HandleReferenceSkeletonChange();
+	}
+#endif
 }
 
 void FReferenceSkeleton::RemoveDuplicateBones(const UObject* Requester, TArray<FBoneIndexType> & DuplicateBones)
diff --git a/Engine/Source/Runtime/Engine/Private/Animation/Skeleton.cpp b/Engine/Source/Runtime/Engine/Private/Animation/Skeleton.cpp
--- a/Engine/Source/Runtime/Engine/Private/Animation/Skeleton.cpp
+++ b/Engine/Source/Runtime/Engine/Private/Animation/Skeleton.cpp
@@ -1476,6 +1476,12 @@ void USkeleton::HandleSkeletonHierarchyChange(bool bShowProgress /*= true*/)
 	OnSkeletonHierarchyChanged.Broadcast();
 }
 
+void USkeleton::HandleReferenceSkeletonChange() const
+{
+	TRACE_CPUPROFILER_EVENT_SCOPE(USkeleton::HandleReferenceSkeletonChange);
+	OnReferenceSkeletonChanged.Broadcast();
+}
+
 void USkeleton::RegisterOnSkeletonHierarchyChanged(const FOnSkeletonHierarchyChanged& Delegate)
 {
 	OnSkeletonHierarchyChanged.Add(Delegate);
@@ -1486,6 +1492,16 @@ void USkeleton::UnregisterOnSkeletonHierarchyChanged(FDelegateUserObject Unregis
 	OnSkeletonHierarchyChanged.RemoveAll(Unregister);
 }
 
+void USkeleton::RegisterOnReferenceSkeletonChanged(const FOnSkeletonHierarchyChanged& Delegate)
+{
+	OnReferenceSkeletonChanged.Add(Delegate);
+}
+
+void USkeleton::UnregisterOnReferenceSkeletonChanged(FDelegateUserObject Unregister)
+{
+	OnReferenceSkeletonChanged.RemoveAll(Unregister);
+}
+
 #endif
 
 #endif // WITH_EDITORONLY_DATA
