Fix UAF automation failure when we GC while holding a scoped temporary node registry

diff --git a/Engine/Plugins/Experimental/UAF/UAFAnimGraph/Source/UAFAnimGraph/Internal/TraitCore/NodeTemplateRegistry.h b/Engine/Plugins/Experimental/UAF/UAFAnimGraph/Source/UAFAnimGraph/Internal/TraitCore/NodeTemplateRegistry.h
--- a/Engine/Plugins/Experimental/UAF/UAFAnimGraph/Source/UAFAnimGraph/Internal/TraitCore/NodeTemplateRegistry.h
+++ b/Engine/Plugins/Experimental/UAF/UAFAnimGraph/Source/UAFAnimGraph/Internal/TraitCore/NodeTemplateRegistry.h
@@ -75,7 +75,11 @@ namespace UE::UAF
 
 		friend class UE::UAF::AnimGraph::FAnimNextAnimGraphModule;
 		friend class FTraitWriter;
+
+#if WITH_DEV_AUTOMATION_TESTS
 		friend struct FScopedClearNodeTemplateRegistry;
+		friend struct FScopedNewNodeTemplateRegistry;
+#endif
 	};
 }
 
diff --git a/Engine/Plugins/Experimental/UAF/UAFTests/Source/UAFAnimGraphTestSuite/Private/AnimNextAnimGraphTraitEventTest.cpp b/Engine/Plugins/Experimental/UAF/UAFTests/Source/UAFAnimGraphTestSuite/Private/AnimNextAnimGraphTraitEventTest.cpp
--- a/Engine/Plugins/Experimental/UAF/UAFTests/Source/UAFAnimGraphTestSuite/Private/AnimNextAnimGraphTraitEventTest.cpp
+++ b/Engine/Plugins/Experimental/UAF/UAFTests/Source/UAFAnimGraphTestSuite/Private/AnimNextAnimGraphTraitEventTest.cpp
@@ -413,12 +413,6 @@ bool FAnimationAnimNextRuntimeTest_GraphTraitEvent::RunTest(const FString& InPar
 				AddErrorIfFalse(EventA->VisitedTraits[0] == NodeTemplateTraitListB[0], "FAnimationAnimNextRuntimeTest_GraphTraitEvent -> Unexpected trait visited");
 			}
 		}
-
-		Registry.Unregister(NodeTemplateA);
-		Registry.Unregister(NodeTemplateB);
-		Registry.Unregister(NodeTemplateC);
-
-		AddErrorIfFalse(Registry.GetNum() == 0, "FAnimationAnimNextRuntimeTest_GraphTraitEvent -> Registry should contain 0 templates");
 	}
 
 	Tests::FUtils::CleanupAfterTests();
diff --git a/Engine/Plugins/Experimental/UAF/UAFTests/Source/UAFAnimGraphTestSuite/Private/AnimNextAnimGraphTraitInterfacesTest.cpp b/Engine/Plugins/Experimental/UAF/UAFTests/Source/UAFAnimGraphTestSuite/Private/AnimNextAnimGraphTraitInterfacesTest.cpp
--- a/Engine/Plugins/Experimental/UAF/UAFTests/Source/UAFAnimGraphTestSuite/Private/AnimNextAnimGraphTraitInterfacesTest.cpp
+++ b/Engine/Plugins/Experimental/UAF/UAFTests/Source/UAFAnimGraphTestSuite/Private/AnimNextAnimGraphTraitInterfacesTest.cpp
@@ -624,13 +624,6 @@ bool FAnimationAnimNextRuntimeTest_IHierarchy::RunTest(const FString& InParamete
 				}
 			}
 		}
-
-		Registry.Unregister(NodeTemplateA);
-		Registry.Unregister(NodeTemplateB);
-		Registry.Unregister(NodeTemplateC);
-		Registry.Unregister(NodeTemplateD);
-
-		AddErrorIfFalse(Registry.GetNum() == 0, "FAnimationAnimNextRuntimeTest_IHierarchy -> Registry should contain 0 templates");
 	}
 
 	Tests::FUtils::CleanupAfterTests();
@@ -768,12 +761,6 @@ bool FAnimationAnimNextRuntimeTest_IUpdate::RunTest(const FString& InParameters)
 
 			Private::UpdatedTraits = nullptr;
 		}
-
-		Registry.Unregister(NodeTemplateA);
-		Registry.Unregister(NodeTemplateB);
-		Registry.Unregister(NodeTemplateC);
-
-		AddErrorIfFalse(Registry.GetNum() == 0, "FAnimationAnimNextRuntimeTest_IUpdate -> Registry should contain 0 templates");
 	}
 	Tests::FUtils::CleanupAfterTests();
 
@@ -909,12 +896,6 @@ bool FAnimationAnimNextRuntimeTest_IEvaluate::RunTest(const FString& InParameter
 
 			Private::EvaluatedTraits = nullptr;
 		}
-
-		Registry.Unregister(NodeTemplateA);
-		Registry.Unregister(NodeTemplateB);
-		Registry.Unregister(NodeTemplateC);
-
-		AddErrorIfFalse(Registry.GetNum() == 0, "FAnimationAnimNextRuntimeTest_IEvaluate -> Registry should contain 0 templates");
 	}
 	Tests::FUtils::CleanupAfterTests();
 
@@ -1069,11 +1050,6 @@ bool FAnimationAnimNextRuntimeTest_IScopedInterface::RunTest(const FString& InPa
 
 			Private::UpdatedTraits = nullptr;
 		}
-
-		Registry.Unregister(NodeTemplate0);
-		Registry.Unregister(NodeTemplate1);
-
-		AddErrorIfFalse(Registry.GetNum() == 0, "FAnimationAnimNextRuntimeTest_IScopedInterface -> Registry should contain 0 templates");
 	}
 	Tests::FUtils::CleanupAfterTests();
 
diff --git a/Engine/Plugins/Experimental/UAF/UAFTests/Source/UAFAnimGraphTestSuite/Private/AnimNextRuntimeTest.cpp b/Engine/Plugins/Experimental/UAF/UAFTests/Source/UAFAnimGraphTestSuite/Private/AnimNextRuntimeTest.cpp
--- a/Engine/Plugins/Experimental/UAF/UAFTests/Source/UAFAnimGraphTestSuite/Private/AnimNextRuntimeTest.cpp
+++ b/Engine/Plugins/Experimental/UAF/UAFTests/Source/UAFAnimGraphTestSuite/Private/AnimNextRuntimeTest.cpp
@@ -16,10 +16,26 @@ namespace UE::UAF
 {
 	FScopedClearNodeTemplateRegistry::FScopedClearNodeTemplateRegistry()
 	{
-		OriginalRegistry = FNodeTemplateRegistry::Swap(&TmpRegistry);
+		FNodeTemplateRegistry& NodeRegistry = FNodeTemplateRegistry::Get();
+
+		OldTemplateBufferSize = NodeRegistry.TemplateBuffer.Num();
+		OldTemplateUIDToHandleMap = NodeRegistry.TemplateUIDToHandleMap;
 	}
 
 	FScopedClearNodeTemplateRegistry::~FScopedClearNodeTemplateRegistry()
+	{
+		FNodeTemplateRegistry& NodeRegistry = FNodeTemplateRegistry::Get();
+
+		NodeRegistry.TemplateBuffer.SetNum(OldTemplateBufferSize, EAllowShrinking::No);
+		NodeRegistry.TemplateUIDToHandleMap = OldTemplateUIDToHandleMap;
+	}
+
+	FScopedNewNodeTemplateRegistry::FScopedNewNodeTemplateRegistry()
+	{
+		OriginalRegistry = FNodeTemplateRegistry::Swap(&TmpRegistry);
+	}
+
+	FScopedNewNodeTemplateRegistry::~FScopedNewNodeTemplateRegistry()
 	{
 		FNodeTemplateRegistry* PrevRegistry = FNodeTemplateRegistry::Swap(OriginalRegistry);
 		checkf(PrevRegistry == &TmpRegistry, TEXT("Unexpected node template registry instance found"));
diff --git a/Engine/Plugins/Experimental/UAF/UAFTests/Source/UAFAnimGraphTestSuite/Private/AnimNextRuntimeTest.h b/Engine/Plugins/Experimental/UAF/UAFTests/Source/UAFAnimGraphTestSuite/Private/AnimNextRuntimeTest.h
--- a/Engine/Plugins/Experimental/UAF/UAFTests/Source/UAFAnimGraphTestSuite/Private/AnimNextRuntimeTest.h
+++ b/Engine/Plugins/Experimental/UAF/UAFTests/Source/UAFAnimGraphTestSuite/Private/AnimNextRuntimeTest.h
@@ -68,15 +68,32 @@ namespace UE::UAF
 		return FString();
 	}
 
-	// Creates a temporary empty node template registry instance and swaps it for the current one
+	// Captures the node template registry and restores its state when we go out of scope
 	struct FScopedClearNodeTemplateRegistry final
 	{
-		FScopedClearNodeTemplateRegistry(const FScopedClearNodeTemplateRegistry&) = delete;
-		FScopedClearNodeTemplateRegistry& operator=(const FScopedClearNodeTemplateRegistry&) = delete;
+		UE_NONCOPYABLE(FScopedClearNodeTemplateRegistry);	// We disallow copy/move semantics
 
 		UAFANIMGRAPHTESTSUITE_API FScopedClearNodeTemplateRegistry();
 		UAFANIMGRAPHTESTSUITE_API ~FScopedClearNodeTemplateRegistry();
 
+		// The size of the TemplateBuffer in the node registry when the scope began
+		// We'll restore it to that size (it is append only) when the scope ends
+		int32 OldTemplateBufferSize = 0;
+
+		// The TemplateUIDToHandleMap in the node registry when the scope began
+		// We'll restore it to this one when the scope ends
+		TMap<uint32, FNodeTemplateRegistryHandle> OldTemplateUIDToHandleMap;
+	};
+
+	// Creates a temporary empty node template registry instance and swaps it for the current one until the end of the scope
+	// Should only be used when testing the node registry itself
+	struct FScopedNewNodeTemplateRegistry final
+	{
+		UE_NONCOPYABLE(FScopedNewNodeTemplateRegistry);	// We disallow copy/move semantics
+
+		UAFANIMGRAPHTESTSUITE_API FScopedNewNodeTemplateRegistry();
+		UAFANIMGRAPHTESTSUITE_API ~FScopedNewNodeTemplateRegistry();
+
 		// The new temporary node registry instance we'll set as the global one
 		FNodeTemplateRegistry TmpRegistry;
 
diff --git a/Engine/Plugins/Experimental/UAF/UAFTests/Source/UAFAnimGraphTestSuite/Private/AnimNextTraitBaseTest.cpp b/Engine/Plugins/Experimental/UAF/UAFTests/Source/UAFAnimGraphTestSuite/Private/AnimNextTraitBaseTest.cpp
--- a/Engine/Plugins/Experimental/UAF/UAFTests/Source/UAFAnimGraphTestSuite/Private/AnimNextTraitBaseTest.cpp
+++ b/Engine/Plugins/Experimental/UAF/UAFTests/Source/UAFAnimGraphTestSuite/Private/AnimNextTraitBaseTest.cpp
@@ -445,7 +445,7 @@ bool FAnimationAnimNextRuntimeTest_NodeTemplateRegistry::RunTest(const FString&
 	AUTO_REGISTER_ANIM_TRAIT(FTraitAB_Add)
 	AUTO_REGISTER_ANIM_TRAIT(FTraitAC_Add)
 
-	FScopedClearNodeTemplateRegistry ScopedClearNodeTemplateRegistry;
+	FScopedNewNodeTemplateRegistry ScopedClearNodeTemplateRegistry;
 	FNodeTemplateRegistry& Registry = FNodeTemplateRegistry::Get();
 
 	TArray<FTraitUID> NodeTemplateTraitList;
@@ -714,11 +714,6 @@ bool FAnimationAnimNextRuntimeTest_NodeLifetime::RunTest(const FString& InParame
 			Private::ConstructedTraits = nullptr;
 			Private::DestructedTraits = nullptr;
 		}
-
-		// Unregister our templates
-		Registry.Unregister(NodeTemplate0);
-
-		AddErrorIfFalse(Registry.GetNum() == 0, "FAnimationAnimNextRuntimeTest_NodeLifetime -> Registry should contain 0 templates");
 	}
 	Tests::FUtils::CleanupAfterTests();
 	
@@ -1136,10 +1131,6 @@ bool FAnimationAnimNextRuntimeTest_TraitStackBinding::RunTest(const FString& InP
 				}
 			}
 		}
-
-		Registry.Unregister(NodeTemplate0);
-
-		AddErrorIfFalse(Registry.GetNum() == 0, "FAnimationAnimNextRuntimeTest_TraitStackBinding -> Registry should contain 0 templates");
 	}
 
 	Tests::FUtils::CleanupAfterTests();
@@ -1362,7 +1353,7 @@ bool FAnimationAnimNextRuntimeTest_TraitSerialization::RunTest(const FString& In
 
 		// Clear out the node template registry to test registration on load
 		{
-			FScopedClearNodeTemplateRegistry ScopedClearNodeTemplateRegistryForLoad;
+			FScopedNewNodeTemplateRegistry ScopedClearNodeTemplateRegistryForLoad;
 
 			AddErrorIfFalse(FNodeTemplateRegistry::Get().GetNum() == 0, "FAnimationAnimNextRuntimeTest_TraitSerialization -> Registry should contain 0 templates");
 
@@ -1571,7 +1562,7 @@ bool FAnimationAnimNextRuntimeTest_SimpleGraphBuilder::RunTest(const FString& In
 
 			// Clear out the node template registry to test registration on load
 			{
-				FScopedClearNodeTemplateRegistry ScopedClearNodeTemplateRegistryForLoad;
+				FScopedNewNodeTemplateRegistry ScopedClearNodeTemplateRegistryForLoad;
 				FNodeTemplateRegistry& Registry = FNodeTemplateRegistry::Get();
 
 				AddErrorIfFalse(Registry.GetNum() == 0, "FAnimationAnimNextRuntimeTest_SimpleGraphBuilder -> Registry should contain 0 templates");
diff --git a/Engine/Plugins/Experimental/UAF/UAFTests/Source/UAFAnimGraphTestSuite/Private/AnimNextTraitEventTest.cpp b/Engine/Plugins/Experimental/UAF/UAFTests/Source/UAFAnimGraphTestSuite/Private/AnimNextTraitEventTest.cpp
--- a/Engine/Plugins/Experimental/UAF/UAFTests/Source/UAFAnimGraphTestSuite/Private/AnimNextTraitEventTest.cpp
+++ b/Engine/Plugins/Experimental/UAF/UAFTests/Source/UAFAnimGraphTestSuite/Private/AnimNextTraitEventTest.cpp
@@ -418,10 +418,6 @@ bool FAnimationAnimNextRuntimeTest_TraitEventRaising::RunTest(const FString& InP
 				AddErrorIfFalse(EventB->VisitedTraits[2] == NodeTemplateTraitList[0], "FAnimationAnimNextRuntimeTest_TraitEventRaising -> Unexpected trait visited");
 			}
 		}
-
-		Registry.Unregister(NodeTemplate0);
-
-		AddErrorIfFalse(Registry.GetNum() == 0, "FAnimationAnimNextRuntimeTest_TraitEventRaising -> Registry should contain 0 templates");
 	}
 
 	Tests::FUtils::CleanupAfterTests();
