//UE5/Main - Cook EngineTest Server Win64 - Found a deprecated '../../../Engine/Plugins/Experimental/UAF/UAFTests/Config/BaseUAFTests.ini' - Rename to Default<PluginName>.ini'
#fix renamed ini file
#tests preflight cook

diff --git a/Engine/Plugins/Experimental/UAF/UAFTests/Config/BaseUAFTests.ini b/Engine/Plugins/Experimental/UAF/UAFTests/Config/BaseUAFTests.ini
deleted file mode 100644
diff --git a/Engine/Plugins/Experimental/UAF/UAFTests/Config/DefaultUAFTests.ini b/Engine/Plugins/Experimental/UAF/UAFTests/Config/DefaultUAFTests.ini
new file mode 100644
--- /dev/null
+++ b/Engine/Plugins/Experimental/UAF/UAFTests/Config/DefaultUAFTests.ini
@@ -0,0 +1,2 @@
+﻿[CoreRedirects]
++PropertyRedirects=(OldName="/Script/UAFTestSuite.UAFVariableReferenceTestStruct.BoolVariable",NewName="/Script/UAFTestSuite.UAFVariableReferenceTestStruct.Renamed_BoolVariable")
\ No newline at end of file
