Add support for dynamic attributes in abstract skeleton runtime

diff --git a/Engine/Plugins/Experimental/UAF/UAFTests/Source/UAFTestSuite/Private/UAF/Attributes/AttributeCollectionTest.cpp b/Engine/Plugins/Experimental/UAF/UAFTests/Source/UAFTestSuite/Private/UAF/Attributes/AttributeCollectionTest.cpp
--- a/Engine/Plugins/Experimental/UAF/UAFTests/Source/UAFTestSuite/Private/UAF/Attributes/AttributeCollectionTest.cpp
+++ b/Engine/Plugins/Experimental/UAF/UAFTests/Source/UAFTestSuite/Private/UAF/Attributes/AttributeCollectionTest.cpp
@@ -15,41 +15,12 @@
 #include "UAF/Attributes/AttributeCollection.h"
 #include "UAF/Attributes/AttributeSetMap.h"
 #include "UAF/Attributes/EngineAttributes.h"
+#include "UAF/Attributes/IteratorTestUtils.h"
 #include "UObject/UObjectGlobals.h"
 #include "UObject/Package.h"
 
 namespace UE::UAF::Tests
 {
-	template <typename IteratorType, typename ValueType, typename ProjectionType>
-	constexpr bool ContainsBy(IteratorType&& It, const ValueType& Value, ProjectionType Proj)
-	{
-		while (It)
-		{
-			if (Invoke(Proj, It) == Value)
-			{
-				return true;
-			}
-
-			++It;
-		}
-
-		return false;
-	}
-
-	template <typename IteratorType>
-	constexpr int32 IteratorSize(IteratorType&& It)
-	{
-		int32 Size = 0;
-
-		while (It)
-		{
-			++Size;
-			++It;
-		}
-
-		return Size;
-	}
-
 	IMPLEMENT_SIMPLE_AUTOMATION_TEST(FAttributeCollectionTest, "Animation.UAF.Attributes.AttributeCollection", EAutomationTestFlags::EditorContext | EAutomationTestFlags::EngineFilter)
 
 	bool FAttributeCollectionTest::RunTest(const FString& InParameters)
@@ -129,8 +100,14 @@ namespace UE::UAF::Tests
 			UE_RETURN_ON_ERROR(FloatToFloatMap != nullptr, TEXT("Float->Float map should be valid"));
 
 			FAttributeMappingKey FloatToIntegerKey = FAttributeMappingKey::MakeFromTo<FFloatAnimationAttribute, FIntegerAnimationAttribute>();
-			TAttributeSetMap<FIntegerAnimationAttribute>* FloatToIntegerMap = Collection.AddSetMap<FIntegerAnimationAttribute>(FloatToIntegerKey);
+			TAttributeSetMap<FIntegerAnimationAttribute>* FloatToIntegerMap = MakeAttributeSetMap<FIntegerAnimationAttribute>(FAttributeSetMap::FConstructArgs(FullBodySet->FindTypedSet<FFloatAnimationAttribute>(), FIntegerAnimationAttribute::StaticStruct(), Collection.GetAllocator()));
 			UE_RETURN_ON_ERROR(FloatToIntegerMap != nullptr, TEXT("Float->Integer map should be valid"));
+			bool bAdded = Collection.AddSetMap(FloatToIntegerMap);
+			UE_RETURN_ON_ERROR(bAdded, TEXT("Float->Integer map should be added"));
+
+			// Can't add twice
+			bAdded = Collection.AddSetMap(FloatToIntegerMap);
+			AddErrorIfFalse(!bAdded, TEXT("Float->Integer map should not be added"));
 
 			FAttributeMappingKey FloatToStringKey = FAttributeMappingKey::MakeFromTo<FFloatAnimationAttribute, FStringAnimationAttribute>();
 			TAttributeSetMap<FStringAnimationAttribute>* FloatToStringMap = Collection.AddSetMap<FStringAnimationAttribute>(FloatToStringKey);
@@ -141,32 +118,56 @@ namespace UE::UAF::Tests
 			UE_RETURN_ON_ERROR(TransformToTransformMap != nullptr, TEXT("Transform->Transform map should be valid"));
 
 			FAttributeMappingKey TransformToFloatKey = FAttributeMappingKey::MakeFromTo<FBoneTransformAnimationAttribute, FFloatAnimationAttribute>();
-			TAttributeSetMap<FFloatAnimationAttribute>* TransformToFloatMap = Collection.AddSetMap<FFloatAnimationAttribute>(TransformToFloatKey);
+			TAttributeSetMap<FFloatAnimationAttribute>* TransformToFloatMap = Cast<FFloatAnimationAttribute>(Collection.AddSetMap(TransformToFloatKey));
 			UE_RETURN_ON_ERROR(TransformToFloatMap != nullptr, TEXT("Transform->Bone map should be valid"));
 
+			TAttributeMap<FFloatAnimationAttribute>* FloatMap = Collection.AddMap<FFloatAnimationAttribute>();
+			UE_RETURN_ON_ERROR(FloatMap != nullptr, TEXT("Float map should be valid"));
+
+			TAttributeMap<FBoneTransformAnimationAttribute>* BoneMap = Cast<FBoneTransformAnimationAttribute>(Collection.AddMap(FBoneTransformAnimationAttribute::StaticStruct()));
+			UE_RETURN_ON_ERROR(BoneMap != nullptr, TEXT("Bone map should be valid"));
+
+			TAttributeMap<FIntegerAnimationAttribute>* IntMap = MakeAttributeMap<FIntegerAnimationAttribute>(Collection.GetAllocator());
+			UE_RETURN_ON_ERROR(IntMap != nullptr, TEXT("Integer map should be valid"));
+			bAdded = Collection.AddMap(IntMap);
+			UE_RETURN_ON_ERROR(bAdded, TEXT("Integer map should be added"));
+
+			// Can't add twice
+			bAdded = Collection.AddMap(IntMap);
+			AddErrorIfFalse(!bAdded, TEXT("Integer map should not be added"));
+
 			AddErrorIfFalse(!Collection.IsEmpty(), TEXT("Attribute collection should not be empty"));
 			AddErrorIfFalse(Collection.NumSetMaps() == 5, TEXT("Unexpected set map count"));
+			AddErrorIfFalse(Collection.NumMaps() == 3, TEXT("Unexpected map count"));
 
 			AddErrorIfFalse(FloatToFloatMap == Collection.FindSetMap<FFloatAnimationAttribute>(FloatToFloatKey), TEXT("Unexpected result from FindSetMap"));
 			AddErrorIfFalse(FloatToIntegerMap == Collection.FindSetMap<FIntegerAnimationAttribute>(FloatToIntegerKey), TEXT("Unexpected result from FindSetMap"));
 			AddErrorIfFalse(FloatToStringMap == Collection.FindSetMap<FStringAnimationAttribute>(FloatToStringKey), TEXT("Unexpected result from FindSetMap"));
 			AddErrorIfFalse(TransformToTransformMap == Collection.FindSetMap<FBoneTransformAnimationAttribute>(TransformToTransformKey), TEXT("Unexpected result from FindSetMap"));
-			AddErrorIfFalse(TransformToFloatMap == Collection.FindSetMap<FFloatAnimationAttribute>(TransformToFloatKey), TEXT("Unexpected result from FindSetMap"));
+			AddErrorIfFalse(TransformToFloatMap == Collection.FindSetMap(TransformToFloatKey), TEXT("Unexpected result from FindSetMap"));
+			AddErrorIfFalse(FloatMap == Collection.FindMap<FFloatAnimationAttribute>(), TEXT("Unexpected result from FindMap"));
+			AddErrorIfFalse(BoneMap == Collection.FindMap(FBoneTransformAnimationAttribute::StaticStruct()), TEXT("Unexpected result from FindMap"));
 
 			AddErrorIfFalse(Collection.RemoveSetMap(FloatToFloatMap), TEXT("Failed to remove set map"));
 			AddErrorIfFalse(Collection.RemoveSetMap(FloatToIntegerKey), TEXT("Failed to remove set map"));
 			AddErrorIfFalse(Collection.RemoveSetMap(FloatToStringMap), TEXT("Failed to remove set map"));
 			AddErrorIfFalse(Collection.RemoveSetMap(TransformToTransformKey), TEXT("Failed to remove set map"));
-			AddErrorIfFalse(Collection.RemoveSetMap(TransformToFloatMap), TEXT("Failed to remove set map"));
+			AddErrorIfFalse(Collection.RemoveSetMap(TransformToFloatKey), TEXT("Failed to remove set map"));
+			AddErrorIfFalse(Collection.RemoveMap(FloatMap), TEXT("Failed to remove map"));
+			AddErrorIfFalse(Collection.RemoveMap(FBoneTransformAnimationAttribute::StaticStruct()), TEXT("Failed to remove map"));
+			AddErrorIfFalse(Collection.RemoveMap(IntMap), TEXT("Failed to remove map"));
 
 			AddErrorIfFalse(nullptr == Collection.FindSetMap<FFloatAnimationAttribute>(FloatToFloatKey), TEXT("Unexpected result from FindSetMap"));
 			AddErrorIfFalse(nullptr == Collection.FindSetMap<FIntegerAnimationAttribute>(FloatToIntegerKey), TEXT("Unexpected result from FindSetMap"));
 			AddErrorIfFalse(nullptr == Collection.FindSetMap<FStringAnimationAttribute>(FloatToStringKey), TEXT("Unexpected result from FindSetMap"));
 			AddErrorIfFalse(nullptr == Collection.FindSetMap<FBoneTransformAnimationAttribute>(TransformToTransformKey), TEXT("Unexpected result from FindSetMap"));
-			AddErrorIfFalse(nullptr == Collection.FindSetMap<FFloatAnimationAttribute>(TransformToFloatKey), TEXT("Unexpected result from FindSetMap"));
+			AddErrorIfFalse(nullptr == Collection.FindSetMap(TransformToFloatKey), TEXT("Unexpected result from FindSetMap"));
+			AddErrorIfFalse(nullptr == Collection.FindMap<FFloatAnimationAttribute>(), TEXT("Unexpected result from FindMap"));
+			AddErrorIfFalse(nullptr == Collection.FindMap(FBoneTransformAnimationAttribute::StaticStruct()), TEXT("Unexpected result from FindMap"));
 
 			AddErrorIfFalse(Collection.IsEmpty(), TEXT("Attribute collection should be empty"));
 			AddErrorIfFalse(Collection.NumSetMaps() == 0, TEXT("Unexpected set map count"));
+			AddErrorIfFalse(Collection.NumMaps() == 0, TEXT("Unexpected map count"));
 
 			FloatToFloatMap = Collection.FindOrAddSetMap<FFloatAnimationAttribute>(FloatToFloatKey);
 			UE_RETURN_ON_ERROR(FloatToFloatMap != nullptr, TEXT("Float->Float map should be valid"));
@@ -180,20 +181,88 @@ namespace UE::UAF::Tests
 			TransformToTransformMap = Collection.FindOrAddSetMap<FBoneTransformAnimationAttribute>(TransformToTransformKey);
 			UE_RETURN_ON_ERROR(TransformToTransformMap != nullptr, TEXT("Transform->Transform map should be valid"));
 
-			TransformToFloatMap = Collection.FindOrAddSetMap<FFloatAnimationAttribute>(TransformToFloatKey);
+			TransformToFloatMap = Cast<FFloatAnimationAttribute>(Collection.FindOrAddSetMap(TransformToFloatKey));
 			UE_RETURN_ON_ERROR(TransformToFloatMap != nullptr, TEXT("Transform->Bone map should be valid"));
 
+			FloatMap = Collection.FindOrAddMap<FFloatAnimationAttribute>();
+			UE_RETURN_ON_ERROR(FloatMap != nullptr, TEXT("Float map should be valid"));
+
+			BoneMap = Cast<FBoneTransformAnimationAttribute>(Collection.FindOrAddMap(FBoneTransformAnimationAttribute::StaticStruct()));
+			UE_RETURN_ON_ERROR(BoneMap != nullptr, TEXT("Bone map should be valid"));
+
 			AddErrorIfFalse(!Collection.IsEmpty(), TEXT("Attribute collection should not be empty"));
 			AddErrorIfFalse(Collection.NumSetMaps() == 5, TEXT("Unexpected set map count"));
+			AddErrorIfFalse(Collection.NumMaps() == 2, TEXT("Unexpected map count"));
 
 			AddErrorIfFalse(FloatToFloatMap == Collection.FindOrAddSetMap<FFloatAnimationAttribute>(FloatToFloatKey), TEXT("Unexpected result from FindOrAddSetMap"));
 			AddErrorIfFalse(FloatToIntegerMap == Collection.FindOrAddSetMap<FIntegerAnimationAttribute>(FloatToIntegerKey), TEXT("Unexpected result from FindOrAddSetMap"));
 			AddErrorIfFalse(FloatToStringMap == Collection.FindOrAddSetMap<FStringAnimationAttribute>(FloatToStringKey), TEXT("Unexpected result from FindOrAddSetMap"));
 			AddErrorIfFalse(TransformToTransformMap == Collection.FindOrAddSetMap<FBoneTransformAnimationAttribute>(TransformToTransformKey), TEXT("Unexpected result from FindOrAddSetMap"));
-			AddErrorIfFalse(TransformToFloatMap == Collection.FindOrAddSetMap<FFloatAnimationAttribute>(TransformToFloatKey), TEXT("Unexpected result from FindOrAddSetMap"));
+			AddErrorIfFalse(TransformToFloatMap == Collection.FindOrAddSetMap(TransformToFloatKey), TEXT("Unexpected result from FindOrAddSetMap"));
+			AddErrorIfFalse(FloatMap == Collection.FindOrAddMap<FFloatAnimationAttribute>(), TEXT("Unexpected result from FindOrAddMap"));
+			AddErrorIfFalse(BoneMap == Collection.FindOrAddMap(FBoneTransformAnimationAttribute::StaticStruct()), TEXT("Unexpected result from FindOrAddMap"));
 
 			AddErrorIfFalse(!Collection.IsEmpty(), TEXT("Attribute collection should not be empty"));
 			AddErrorIfFalse(Collection.NumSetMaps() == 5, TEXT("Unexpected set map count"));
+			AddErrorIfFalse(Collection.NumMaps() == 2, TEXT("Unexpected map count"));
+
+			// Append
+			{
+				FAttributeCollectionHeap CollectionCopy(FullBodySet);
+
+				for (auto It = Collection.CreateConstSetMapIterator(); It; ++It)
+				{
+					FAttributeSetMap* SetMapCopy = It.GetMap()->Duplicate(CollectionCopy.GetAllocator());
+					const bool bAppended = CollectionCopy.AppendSetMap(SetMapCopy);
+					AddErrorIfFalse(bAppended, TEXT("Set map should be appended"));
+				}
+
+				for (auto It = Collection.CreateConstMapIterator(); It; ++It)
+				{
+					FAttributeMap* MapCopy = It.GetMap()->Duplicate(CollectionCopy.GetAllocator());
+					const bool bAppended = CollectionCopy.AppendMap(MapCopy);
+					AddErrorIfFalse(bAppended, TEXT("Map should be appended"));
+				}
+
+				AddErrorIfFalse(CollectionCopy.IsEmpty() == Collection.IsEmpty(), TEXT("Attribute collection should not be empty"));
+				AddErrorIfFalse(CollectionCopy.NumSetMaps() == Collection.NumSetMaps(), TEXT("Unexpected set map count"));
+				AddErrorIfFalse(CollectionCopy.NumMaps() == Collection.NumMaps(), TEXT("Unexpected map count"));
+
+				TAttributeSetMap<FFloatAnimationAttribute>* FloatToFloatMapCopy = CollectionCopy.FindSetMap<FFloatAnimationAttribute>(FloatToFloatKey);
+				UE_RETURN_ON_ERROR(FloatToFloatMapCopy != nullptr, TEXT("Float->Float map should be valid"));
+
+				TAttributeSetMap<FIntegerAnimationAttribute>* FloatToIntegerMapCopy = CollectionCopy.FindSetMap<FIntegerAnimationAttribute>(FloatToIntegerKey);
+				UE_RETURN_ON_ERROR(FloatToIntegerMapCopy != nullptr, TEXT("Float->Integer map should be valid"));
+
+				TAttributeSetMap<FStringAnimationAttribute>* FloatToStringMapCopy = CollectionCopy.FindSetMap<FStringAnimationAttribute>(FloatToStringKey);
+				UE_RETURN_ON_ERROR(FloatToStringMapCopy != nullptr, TEXT("Float->String map should be valid"));
+
+				TAttributeSetMap<FBoneTransformAnimationAttribute>* TransformToTransformMapCopy = CollectionCopy.FindSetMap<FBoneTransformAnimationAttribute>(TransformToTransformKey);
+				UE_RETURN_ON_ERROR(TransformToTransformMapCopy != nullptr, TEXT("Transform->Transform map should be valid"));
+
+				TAttributeSetMap<FFloatAnimationAttribute>* TransformToFloatMapCopy = Cast<FFloatAnimationAttribute>(CollectionCopy.FindSetMap(TransformToFloatKey));
+				UE_RETURN_ON_ERROR(TransformToFloatMapCopy != nullptr, TEXT("Transform->Bone map should be valid"));
+
+				TAttributeMap<FFloatAnimationAttribute>* FloatMapCopy = CollectionCopy.FindMap<FFloatAnimationAttribute>();
+				UE_RETURN_ON_ERROR(FloatMapCopy != nullptr, TEXT("Float map should be valid"));
+
+				TAttributeMap<FBoneTransformAnimationAttribute>* BoneMapCopy = Cast<FBoneTransformAnimationAttribute>(CollectionCopy.FindMap(FBoneTransformAnimationAttribute::StaticStruct()));
+				UE_RETURN_ON_ERROR(BoneMapCopy != nullptr, TEXT("Bone map should be valid"));
+
+				AddErrorIfFalse(FloatToFloatMap != FloatToFloatMapCopy, TEXT("Set map instances should be duplicated"));
+				AddErrorIfFalse(FloatToIntegerMap != FloatToIntegerMapCopy, TEXT("Set map instances should be duplicated"));
+				AddErrorIfFalse(FloatToStringMap != FloatToStringMapCopy, TEXT("Set map instances should be duplicated"));
+				AddErrorIfFalse(TransformToTransformMap != TransformToTransformMapCopy, TEXT("Set map instances should be duplicated"));
+				AddErrorIfFalse(TransformToFloatMap != TransformToFloatMapCopy, TEXT("Set map instances should be duplicated"));
+				AddErrorIfFalse(FloatMap != FloatMapCopy, TEXT("Name/value map instances should be duplicated"));
+				AddErrorIfFalse(BoneMap != BoneMapCopy, TEXT("Name/value map instances should be duplicated"));
+
+				AddErrorIfFalse(FloatToFloatMap->GetTypedSet() == FloatToFloatMapCopy->GetTypedSet(), TEXT("Typed set mismatched"));
+				AddErrorIfFalse(FloatToIntegerMap->GetTypedSet() == FloatToIntegerMapCopy->GetTypedSet(), TEXT("Typed set mismatched"));
+				AddErrorIfFalse(FloatToStringMap->GetTypedSet() == FloatToStringMapCopy->GetTypedSet(), TEXT("Typed set mismatched"));
+				AddErrorIfFalse(TransformToTransformMap->GetTypedSet() == TransformToTransformMapCopy->GetTypedSet(), TEXT("Typed set mismatched"));
+				AddErrorIfFalse(TransformToFloatMap->GetTypedSet() == TransformToFloatMapCopy->GetTypedSet(), TEXT("Typed set mismatched"));
+			}
 
 			// Copy constructor
 			{
@@ -201,6 +270,7 @@ namespace UE::UAF::Tests
 
 				AddErrorIfFalse(CollectionCopy.IsEmpty() == Collection.IsEmpty(), TEXT("Attribute collection should not be empty"));
 				AddErrorIfFalse(CollectionCopy.NumSetMaps() == Collection.NumSetMaps(), TEXT("Unexpected set map count"));
+				AddErrorIfFalse(CollectionCopy.NumMaps() == Collection.NumMaps(), TEXT("Unexpected map count"));
 
 				TAttributeSetMap<FFloatAnimationAttribute>* FloatToFloatMapCopy = CollectionCopy.FindSetMap<FFloatAnimationAttribute>(FloatToFloatKey);
 				UE_RETURN_ON_ERROR(FloatToFloatMapCopy != nullptr, TEXT("Float->Float map should be valid"));
@@ -214,14 +284,22 @@ namespace UE::UAF::Tests
 				TAttributeSetMap<FBoneTransformAnimationAttribute>* TransformToTransformMapCopy = CollectionCopy.FindSetMap<FBoneTransformAnimationAttribute>(TransformToTransformKey);
 				UE_RETURN_ON_ERROR(TransformToTransformMapCopy != nullptr, TEXT("Transform->Transform map should be valid"));
 
-				TAttributeSetMap<FFloatAnimationAttribute>* TransformToFloatMapCopy = CollectionCopy.FindSetMap<FFloatAnimationAttribute>(TransformToFloatKey);
+				TAttributeSetMap<FFloatAnimationAttribute>* TransformToFloatMapCopy = Cast<FFloatAnimationAttribute>(CollectionCopy.FindSetMap(TransformToFloatKey));
 				UE_RETURN_ON_ERROR(TransformToFloatMapCopy != nullptr, TEXT("Transform->Bone map should be valid"));
 
+				TAttributeMap<FFloatAnimationAttribute>* FloatMapCopy = CollectionCopy.FindMap<FFloatAnimationAttribute>();
+				UE_RETURN_ON_ERROR(FloatMapCopy != nullptr, TEXT("Float map should be valid"));
+
+				TAttributeMap<FBoneTransformAnimationAttribute>* BoneMapCopy = Cast<FBoneTransformAnimationAttribute>(CollectionCopy.FindMap(FBoneTransformAnimationAttribute::StaticStruct()));
+				UE_RETURN_ON_ERROR(BoneMapCopy != nullptr, TEXT("Bone map should be valid"));
+
 				AddErrorIfFalse(FloatToFloatMap != FloatToFloatMapCopy, TEXT("Set map instances should be duplicated"));
 				AddErrorIfFalse(FloatToIntegerMap != FloatToIntegerMapCopy, TEXT("Set map instances should be duplicated"));
 				AddErrorIfFalse(FloatToStringMap != FloatToStringMapCopy, TEXT("Set map instances should be duplicated"));
 				AddErrorIfFalse(TransformToTransformMap != TransformToTransformMapCopy, TEXT("Set map instances should be duplicated"));
 				AddErrorIfFalse(TransformToFloatMap != TransformToFloatMapCopy, TEXT("Set map instances should be duplicated"));
+				AddErrorIfFalse(FloatMap != FloatMapCopy, TEXT("Name/value map instances should be duplicated"));
+				AddErrorIfFalse(BoneMap != BoneMapCopy, TEXT("Name/value map instances should be duplicated"));
 
 				AddErrorIfFalse(FloatToFloatMap->GetTypedSet() == FloatToFloatMapCopy->GetTypedSet(), TEXT("Typed set mismatched"));
 				AddErrorIfFalse(FloatToIntegerMap->GetTypedSet() == FloatToIntegerMapCopy->GetTypedSet(), TEXT("Typed set mismatched"));
@@ -235,7 +313,8 @@ namespace UE::UAF::Tests
 
 					AddErrorIfFalse(CollectionCopy.IsEmpty(), TEXT("Attribute collection should be empty"));
 					AddErrorIfFalse(!CollectionMoved.IsEmpty(), TEXT("Attribute collection should not be empty"));
-					AddErrorIfFalse(CollectionMoved.NumSetMaps() != CollectionCopy.NumSetMaps(), TEXT("Unexpected set map count"));
+					AddErrorIfFalse(CollectionMoved.NumSetMaps() == Collection.NumSetMaps(), TEXT("Unexpected set map count"));
+					AddErrorIfFalse(CollectionMoved.NumMaps() == Collection.NumMaps(), TEXT("Unexpected map count"));
 
 					TAttributeSetMap<FFloatAnimationAttribute>* FloatToFloatMapMoved = CollectionMoved.FindSetMap<FFloatAnimationAttribute>(FloatToFloatKey);
 					UE_RETURN_ON_ERROR(FloatToFloatMapMoved != nullptr, TEXT("Float->Float map should be valid"));
@@ -249,14 +328,22 @@ namespace UE::UAF::Tests
 					TAttributeSetMap<FBoneTransformAnimationAttribute>* TransformToTransformMapMoved = CollectionMoved.FindSetMap<FBoneTransformAnimationAttribute>(TransformToTransformKey);
 					UE_RETURN_ON_ERROR(TransformToTransformMapMoved != nullptr, TEXT("Transform->Transform map should be valid"));
 
-					TAttributeSetMap<FFloatAnimationAttribute>* TransformToFloatMapMoved = CollectionMoved.FindSetMap<FFloatAnimationAttribute>(TransformToFloatKey);
+					TAttributeSetMap<FFloatAnimationAttribute>* TransformToFloatMapMoved = Cast<FFloatAnimationAttribute>(CollectionMoved.FindSetMap(TransformToFloatKey));
 					UE_RETURN_ON_ERROR(TransformToFloatMapMoved != nullptr, TEXT("Transform->Bone map should be valid"));
 
+					TAttributeMap<FFloatAnimationAttribute>* FloatMapMoved = CollectionMoved.FindMap<FFloatAnimationAttribute>();
+					UE_RETURN_ON_ERROR(FloatMapMoved != nullptr, TEXT("Float map should be valid"));
+
+					TAttributeMap<FBoneTransformAnimationAttribute>* BoneMapMoved = Cast<FBoneTransformAnimationAttribute>(CollectionMoved.FindMap(FBoneTransformAnimationAttribute::StaticStruct()));
+					UE_RETURN_ON_ERROR(BoneMapMoved != nullptr, TEXT("Bone map should be valid"));
+
 					AddErrorIfFalse(FloatToFloatMapMoved == FloatToFloatMapCopy, TEXT("Set map instances should be moved"));
 					AddErrorIfFalse(FloatToIntegerMapMoved == FloatToIntegerMapCopy, TEXT("Set map instances should be moved"));
 					AddErrorIfFalse(FloatToStringMapMoved == FloatToStringMapCopy, TEXT("Set map instances should be moved"));
 					AddErrorIfFalse(TransformToTransformMapMoved == TransformToTransformMapCopy, TEXT("Set map instances should be moved"));
 					AddErrorIfFalse(TransformToFloatMapMoved == TransformToFloatMapCopy, TEXT("Set map instances should be moved"));
+					AddErrorIfFalse(FloatMapMoved == FloatMapCopy, TEXT("Set map instances should be moved"));
+					AddErrorIfFalse(BoneMapMoved == BoneMapCopy, TEXT("Set map instances should be moved"));
 				}
 			}
 
@@ -267,6 +354,7 @@ namespace UE::UAF::Tests
 
 				AddErrorIfFalse(CollectionCopy.IsEmpty() == Collection.IsEmpty(), TEXT("Attribute collection should not be empty"));
 				AddErrorIfFalse(CollectionCopy.NumSetMaps() == Collection.NumSetMaps(), TEXT("Unexpected set map count"));
+				AddErrorIfFalse(CollectionCopy.NumMaps() == Collection.NumMaps(), TEXT("Unexpected map count"));
 
 				TAttributeSetMap<FFloatAnimationAttribute>* FloatToFloatMapCopy = CollectionCopy.FindSetMap<FFloatAnimationAttribute>(FloatToFloatKey);
 				UE_RETURN_ON_ERROR(FloatToFloatMapCopy != nullptr, TEXT("Float->Float map should be valid"));
@@ -280,14 +368,22 @@ namespace UE::UAF::Tests
 				TAttributeSetMap<FBoneTransformAnimationAttribute>* TransformToTransformMapCopy = CollectionCopy.FindSetMap<FBoneTransformAnimationAttribute>(TransformToTransformKey);
 				UE_RETURN_ON_ERROR(TransformToTransformMapCopy != nullptr, TEXT("Transform->Transform map should be valid"));
 
-				TAttributeSetMap<FFloatAnimationAttribute>* TransformToFloatMapCopy = CollectionCopy.FindSetMap<FFloatAnimationAttribute>(TransformToFloatKey);
+				TAttributeSetMap<FFloatAnimationAttribute>* TransformToFloatMapCopy = Cast<FFloatAnimationAttribute>(CollectionCopy.FindSetMap(TransformToFloatKey));
 				UE_RETURN_ON_ERROR(TransformToFloatMapCopy != nullptr, TEXT("Transform->Bone map should be valid"));
 
+				TAttributeMap<FFloatAnimationAttribute>* FloatMapCopy = CollectionCopy.FindMap<FFloatAnimationAttribute>();
+				UE_RETURN_ON_ERROR(FloatMapCopy != nullptr, TEXT("Float map should be valid"));
+
+				TAttributeMap<FBoneTransformAnimationAttribute>* BoneMapCopy = Cast<FBoneTransformAnimationAttribute>(CollectionCopy.FindMap(FBoneTransformAnimationAttribute::StaticStruct()));
+				UE_RETURN_ON_ERROR(BoneMapCopy != nullptr, TEXT("Bone map should be valid"));
+
 				AddErrorIfFalse(FloatToFloatMap != FloatToFloatMapCopy, TEXT("Set map instances should be duplicated"));
 				AddErrorIfFalse(FloatToIntegerMap != FloatToIntegerMapCopy, TEXT("Set map instances should be duplicated"));
 				AddErrorIfFalse(FloatToStringMap != FloatToStringMapCopy, TEXT("Set map instances should be duplicated"));
 				AddErrorIfFalse(TransformToTransformMap != TransformToTransformMapCopy, TEXT("Set map instances should be duplicated"));
 				AddErrorIfFalse(TransformToFloatMap != TransformToFloatMapCopy, TEXT("Set map instances should be duplicated"));
+				AddErrorIfFalse(FloatMap != FloatMapCopy, TEXT("Name/value map instances should be duplicated"));
+				AddErrorIfFalse(BoneMap != BoneMapCopy, TEXT("Name/value map instances should be duplicated"));
 
 				AddErrorIfFalse(FloatToFloatMap->GetTypedSet() == FloatToFloatMapCopy->GetTypedSet(), TEXT("Typed set mismatched"));
 				AddErrorIfFalse(FloatToIntegerMap->GetTypedSet() == FloatToIntegerMapCopy->GetTypedSet(), TEXT("Typed set mismatched"));
@@ -302,7 +398,8 @@ namespace UE::UAF::Tests
 
 					AddErrorIfFalse(CollectionCopy.IsEmpty(), TEXT("Attribute collection should be empty"));
 					AddErrorIfFalse(!CollectionMoved.IsEmpty(), TEXT("Attribute collection should not be empty"));
-					AddErrorIfFalse(CollectionMoved.NumSetMaps() != CollectionCopy.NumSetMaps(), TEXT("Unexpected set map count"));
+					AddErrorIfFalse(CollectionMoved.NumSetMaps() == Collection.NumSetMaps(), TEXT("Unexpected set map count"));
+					AddErrorIfFalse(CollectionMoved.NumMaps() == Collection.NumMaps(), TEXT("Unexpected map count"));
 
 					TAttributeSetMap<FFloatAnimationAttribute>* FloatToFloatMapMoved = CollectionMoved.FindSetMap<FFloatAnimationAttribute>(FloatToFloatKey);
 					UE_RETURN_ON_ERROR(FloatToFloatMapMoved != nullptr, TEXT("Float->Float map should be valid"));
@@ -316,45 +413,83 @@ namespace UE::UAF::Tests
 					TAttributeSetMap<FBoneTransformAnimationAttribute>* TransformToTransformMapMoved = CollectionMoved.FindSetMap<FBoneTransformAnimationAttribute>(TransformToTransformKey);
 					UE_RETURN_ON_ERROR(TransformToTransformMapMoved != nullptr, TEXT("Transform->Transform map should be valid"));
 
-					TAttributeSetMap<FFloatAnimationAttribute>* TransformToFloatMapMoved = CollectionMoved.FindSetMap<FFloatAnimationAttribute>(TransformToFloatKey);
+					TAttributeSetMap<FFloatAnimationAttribute>* TransformToFloatMapMoved = Cast<FFloatAnimationAttribute>(CollectionMoved.FindSetMap(TransformToFloatKey));
 					UE_RETURN_ON_ERROR(TransformToFloatMapMoved != nullptr, TEXT("Transform->Bone map should be valid"));
 
+					TAttributeMap<FFloatAnimationAttribute>* FloatMapMoved = CollectionMoved.FindMap<FFloatAnimationAttribute>();
+					UE_RETURN_ON_ERROR(FloatMapMoved != nullptr, TEXT("Float map should be valid"));
+
+					TAttributeMap<FBoneTransformAnimationAttribute>* BoneMapMoved = Cast<FBoneTransformAnimationAttribute>(CollectionMoved.FindMap(FBoneTransformAnimationAttribute::StaticStruct()));
+					UE_RETURN_ON_ERROR(BoneMapMoved != nullptr, TEXT("Bone map should be valid"));
+
 					AddErrorIfFalse(FloatToFloatMapMoved == FloatToFloatMapCopy, TEXT("Set map instances should be moved"));
 					AddErrorIfFalse(FloatToIntegerMapMoved == FloatToIntegerMapCopy, TEXT("Set map instances should be moved"));
 					AddErrorIfFalse(FloatToStringMapMoved == FloatToStringMapCopy, TEXT("Set map instances should be moved"));
 					AddErrorIfFalse(TransformToTransformMapMoved == TransformToTransformMapCopy, TEXT("Set map instances should be moved"));
 					AddErrorIfFalse(TransformToFloatMapMoved == TransformToFloatMapCopy, TEXT("Set map instances should be moved"));
+					AddErrorIfFalse(FloatMapMoved == FloatMapCopy, TEXT("Name/value map instances should be moved"));
+					AddErrorIfFalse(BoneMapMoved == BoneMapCopy, TEXT("Name/value map instances should be moved"));
 				}
 			}
 
-			// Iterators
+			// Iterators (set map)
 			{
+				const auto GetSetMapKeyProjection = [](const typename FAttributeCollection::TSetMapIterator& It) { return It.GetKey(); };
+				const auto GetSetMapKeyProjectionConst = [](const typename FAttributeCollection::TConstSetMapIterator& It) { return It.GetKey(); };
+				const auto GetSetMapAttributeTypeProjection = [](const typename FAttributeCollection::TSetMapIterator& It) { return It.GetAttributeType(); };
+				const auto GetSetMapAttributeTypeProjectionConst = [](const typename FAttributeCollection::TConstSetMapIterator& It) { return It.GetAttributeType(); };
+				const auto GetSetMapProjection = [](const typename FAttributeCollection::TSetMapIterator& It) { return It.GetMap(); };
+				const auto GetSetMapProjectionConst = [](const typename FAttributeCollection::TConstSetMapIterator& It) { return It.GetMap(); };
+
 				AddErrorIfFalse(IteratorSize(Collection.CreateSetMapIterator()) == 5, TEXT("Iterator did not have expected size"));
 				AddErrorIfFalse(IteratorSize(Collection.CreateConstSetMapIterator()) == 5, TEXT("Iterator did not have expected size"));
 
-				AddErrorIfFalse(ContainsBy(Collection.CreateSetMapIterator(), FloatToFloatMap->GetTypedSet()->GetType(), [](const typename FAttributeCollection::TSetMapIterator& It) { return It.GetAttributeType(); }), TEXT("Iterator did not contain expected attribute type"));
-				AddErrorIfFalse(ContainsBy(Collection.CreateSetMapIterator(), FloatToIntegerMap->GetTypedSet()->GetType(), [](const typename FAttributeCollection::TSetMapIterator& It) { return It.GetAttributeType(); }), TEXT("Iterator did not contain expected attribute type"));
-				AddErrorIfFalse(ContainsBy(Collection.CreateSetMapIterator(), FloatToStringMap->GetTypedSet()->GetType(), [](const typename FAttributeCollection::TSetMapIterator& It) { return It.GetAttributeType(); }), TEXT("Iterator did not contain expected attribute type"));
-				AddErrorIfFalse(ContainsBy(Collection.CreateSetMapIterator(), TransformToTransformMap->GetTypedSet()->GetType(), [](const typename FAttributeCollection::TSetMapIterator& It) { return It.GetAttributeType(); }), TEXT("Iterator did not contain expected attribute type"));
-				AddErrorIfFalse(ContainsBy(Collection.CreateSetMapIterator(), TransformToFloatMap->GetTypedSet()->GetType(), [](const typename FAttributeCollection::TSetMapIterator& It) { return It.GetAttributeType(); }), TEXT("Iterator did not contain expected attribute type"));
-
-				AddErrorIfFalse(ContainsBy(Collection.CreateSetMapIterator(), FloatToFloatMap, [](const typename FAttributeCollection::TSetMapIterator& It) { return It.GetMap(); }), TEXT("Iterator did not contain expected set map"));
-				AddErrorIfFalse(ContainsBy(Collection.CreateSetMapIterator(), FloatToIntegerMap, [](const typename FAttributeCollection::TSetMapIterator& It) { return It.GetMap(); }), TEXT("Iterator did not contain expected set map"));
-				AddErrorIfFalse(ContainsBy(Collection.CreateSetMapIterator(), FloatToStringMap, [](const typename FAttributeCollection::TSetMapIterator& It) { return It.GetMap(); }), TEXT("Iterator did not contain expected set map"));
-				AddErrorIfFalse(ContainsBy(Collection.CreateSetMapIterator(), TransformToTransformMap, [](const typename FAttributeCollection::TSetMapIterator& It) { return It.GetMap(); }), TEXT("Iterator did not contain expected set map"));
-				AddErrorIfFalse(ContainsBy(Collection.CreateSetMapIterator(), TransformToFloatMap, [](const typename FAttributeCollection::TSetMapIterator& It) { return It.GetMap(); }), TEXT("Iterator did not contain expected set map"));
-
-				AddErrorIfFalse(ContainsBy(Collection.CreateConstSetMapIterator(), FloatToFloatMap->GetTypedSet()->GetType(), [](const typename FAttributeCollection::TConstSetMapIterator& It) { return It.GetAttributeType(); }), TEXT("Iterator did not contain expected attribute type"));
-				AddErrorIfFalse(ContainsBy(Collection.CreateConstSetMapIterator(), FloatToIntegerMap->GetTypedSet()->GetType(), [](const typename FAttributeCollection::TConstSetMapIterator& It) { return It.GetAttributeType(); }), TEXT("Iterator did not contain expected attribute type"));
-				AddErrorIfFalse(ContainsBy(Collection.CreateConstSetMapIterator(), FloatToStringMap->GetTypedSet()->GetType(), [](const typename FAttributeCollection::TConstSetMapIterator& It) { return It.GetAttributeType(); }), TEXT("Iterator did not contain expected attribute type"));
-				AddErrorIfFalse(ContainsBy(Collection.CreateConstSetMapIterator(), TransformToTransformMap->GetTypedSet()->GetType(), [](const typename FAttributeCollection::TConstSetMapIterator& It) { return It.GetAttributeType(); }), TEXT("Iterator did not contain expected attribute type"));
-				AddErrorIfFalse(ContainsBy(Collection.CreateConstSetMapIterator(), TransformToFloatMap->GetTypedSet()->GetType(), [](const typename FAttributeCollection::TConstSetMapIterator& It) { return It.GetAttributeType(); }), TEXT("Iterator did not contain expected attribute type"));
-
-				AddErrorIfFalse(ContainsBy(Collection.CreateConstSetMapIterator(), FloatToFloatMap, [](const typename FAttributeCollection::TConstSetMapIterator& It) { return It.GetMap(); }), TEXT("Iterator did not contain expected set map"));
-				AddErrorIfFalse(ContainsBy(Collection.CreateConstSetMapIterator(), FloatToIntegerMap, [](const typename FAttributeCollection::TConstSetMapIterator& It) { return It.GetMap(); }), TEXT("Iterator did not contain expected set map"));
-				AddErrorIfFalse(ContainsBy(Collection.CreateConstSetMapIterator(), FloatToStringMap, [](const typename FAttributeCollection::TConstSetMapIterator& It) { return It.GetMap(); }), TEXT("Iterator did not contain expected set map"));
-				AddErrorIfFalse(ContainsBy(Collection.CreateConstSetMapIterator(), TransformToTransformMap, [](const typename FAttributeCollection::TConstSetMapIterator& It) { return It.GetMap(); }), TEXT("Iterator did not contain expected set map"));
-				AddErrorIfFalse(ContainsBy(Collection.CreateConstSetMapIterator(), TransformToFloatMap, [](const typename FAttributeCollection::TConstSetMapIterator& It) { return It.GetMap(); }), TEXT("Iterator did not contain expected set map"));
+				AddErrorIfFalse(IteratorSortedByPredicate(Collection.CreateSetMapIterator(), GetSetMapKeyProjection, TLess<FAttributeMappingKey>()), TEXT("Iterator should be sorted"));
+				AddErrorIfFalse(IteratorSortedByPredicate(Collection.CreateConstSetMapIterator(), GetSetMapKeyProjectionConst, TLess<FAttributeMappingKey>()), TEXT("Iterator should be sorted"));
+
+				AddErrorIfFalse(FindWithBy(Collection.CreateSetMapIterator(), FloatToFloatMap->GetTypedSet()->GetType(), GetSetMapAttributeTypeProjection), TEXT("Iterator did not contain expected attribute type"));
+				AddErrorIfFalse(FindWithBy(Collection.CreateSetMapIterator(), FloatToIntegerMap->GetTypedSet()->GetType(), GetSetMapAttributeTypeProjection), TEXT("Iterator did not contain expected attribute type"));
+				AddErrorIfFalse(FindWithBy(Collection.CreateSetMapIterator(), FloatToStringMap->GetTypedSet()->GetType(), GetSetMapAttributeTypeProjection), TEXT("Iterator did not contain expected attribute type"));
+				AddErrorIfFalse(FindWithBy(Collection.CreateSetMapIterator(), TransformToTransformMap->GetTypedSet()->GetType(), GetSetMapAttributeTypeProjection), TEXT("Iterator did not contain expected attribute type"));
+				AddErrorIfFalse(FindWithBy(Collection.CreateSetMapIterator(), TransformToFloatMap->GetTypedSet()->GetType(), GetSetMapAttributeTypeProjection), TEXT("Iterator did not contain expected attribute type"));
+
+				AddErrorIfFalse(FindWithBy(Collection.CreateSetMapIterator(), FloatToFloatMap, GetSetMapProjection), TEXT("Iterator did not contain expected set map"));
+				AddErrorIfFalse(FindWithBy(Collection.CreateSetMapIterator(), FloatToIntegerMap, GetSetMapProjection), TEXT("Iterator did not contain expected set map"));
+				AddErrorIfFalse(FindWithBy(Collection.CreateSetMapIterator(), FloatToStringMap, GetSetMapProjection), TEXT("Iterator did not contain expected set map"));
+				AddErrorIfFalse(FindWithBy(Collection.CreateSetMapIterator(), TransformToTransformMap, GetSetMapProjection), TEXT("Iterator did not contain expected set map"));
+				AddErrorIfFalse(FindWithBy(Collection.CreateSetMapIterator(), TransformToFloatMap, GetSetMapProjection), TEXT("Iterator did not contain expected set map"));
+
+				AddErrorIfFalse(FindWithBy(Collection.CreateConstSetMapIterator(), FloatToFloatMap->GetTypedSet()->GetType(), GetSetMapAttributeTypeProjectionConst), TEXT("Iterator did not contain expected attribute type"));
+				AddErrorIfFalse(FindWithBy(Collection.CreateConstSetMapIterator(), FloatToIntegerMap->GetTypedSet()->GetType(), GetSetMapAttributeTypeProjectionConst), TEXT("Iterator did not contain expected attribute type"));
+				AddErrorIfFalse(FindWithBy(Collection.CreateConstSetMapIterator(), FloatToStringMap->GetTypedSet()->GetType(), GetSetMapAttributeTypeProjectionConst), TEXT("Iterator did not contain expected attribute type"));
+				AddErrorIfFalse(FindWithBy(Collection.CreateConstSetMapIterator(), TransformToTransformMap->GetTypedSet()->GetType(), GetSetMapAttributeTypeProjectionConst), TEXT("Iterator did not contain expected attribute type"));
+				AddErrorIfFalse(FindWithBy(Collection.CreateConstSetMapIterator(), TransformToFloatMap->GetTypedSet()->GetType(), GetSetMapAttributeTypeProjectionConst), TEXT("Iterator did not contain expected attribute type"));
+
+				AddErrorIfFalse(FindWithBy(Collection.CreateConstSetMapIterator(), FloatToFloatMap, GetSetMapProjectionConst), TEXT("Iterator did not contain expected set map"));
+				AddErrorIfFalse(FindWithBy(Collection.CreateConstSetMapIterator(), FloatToIntegerMap, GetSetMapProjectionConst), TEXT("Iterator did not contain expected set map"));
+				AddErrorIfFalse(FindWithBy(Collection.CreateConstSetMapIterator(), FloatToStringMap, GetSetMapProjectionConst), TEXT("Iterator did not contain expected set map"));
+				AddErrorIfFalse(FindWithBy(Collection.CreateConstSetMapIterator(), TransformToTransformMap, GetSetMapProjectionConst), TEXT("Iterator did not contain expected set map"));
+				AddErrorIfFalse(FindWithBy(Collection.CreateConstSetMapIterator(), TransformToFloatMap, GetSetMapProjectionConst), TEXT("Iterator did not contain expected set map"));
+			}
+
+			// Iterators (name/value map)
+			{
+				const auto GetMapValueTypeProjection = [](const typename FAttributeCollection::TMapIterator& It) { return It.GetValueType(); };
+				const auto GetMapValueTypeProjectionConst = [](const typename FAttributeCollection::TConstMapIterator& It) { return It.GetValueType(); };
+				const auto GetMapProjection = [](const typename FAttributeCollection::TMapIterator& It) { return It.GetMap(); };
+				const auto GetMapProjectionConst = [](const typename FAttributeCollection::TConstMapIterator& It) { return It.GetMap(); };
+
+				AddErrorIfFalse(IteratorSize(Collection.CreateMapIterator()) == 2, TEXT("Iterator did not have expected size"));
+				AddErrorIfFalse(IteratorSize(Collection.CreateConstMapIterator()) == 2, TEXT("Iterator did not have expected size"));
+
+				AddErrorIfFalse(IteratorSortedByPredicate(Collection.CreateMapIterator(), GetMapValueTypeProjection, TLess<UScriptStruct*>()), TEXT("Iterator should be sorted"));
+				AddErrorIfFalse(IteratorSortedByPredicate(Collection.CreateConstMapIterator(), GetMapValueTypeProjectionConst, TLess<UScriptStruct*>()), TEXT("Iterator should be sorted"));
+
+				AddErrorIfFalse(FindWithBy(Collection.CreateMapIterator(), FloatMap, GetMapProjection), TEXT("Iterator did not contain expected map"));
+				AddErrorIfFalse(FindWithBy(Collection.CreateMapIterator(), BoneMap, GetMapProjection), TEXT("Iterator did not contain expected map"));
+
+				AddErrorIfFalse(FindWithBy(Collection.CreateConstMapIterator(), FloatMap, GetMapProjectionConst), TEXT("Iterator did not contain expected map"));
+				AddErrorIfFalse(FindWithBy(Collection.CreateConstMapIterator(), BoneMap, GetMapProjectionConst), TEXT("Iterator did not contain expected map"));
 			}
 		}
 
