Add support for dynamic attributes in abstract skeleton runtime

diff --git a/Engine/Plugins/Experimental/UAF/UAFTests/Source/UAFTestSuite/Private/UAF/Attributes/IteratorUtilsTest.cpp b/Engine/Plugins/Experimental/UAF/UAFTests/Source/UAFTestSuite/Private/UAF/Attributes/IteratorUtilsTest.cpp
--- a/Engine/Plugins/Experimental/UAF/UAFTests/Source/UAFTestSuite/Private/UAF/Attributes/IteratorUtilsTest.cpp
+++ b/Engine/Plugins/Experimental/UAF/UAFTests/Source/UAFTestSuite/Private/UAF/Attributes/IteratorUtilsTest.cpp
@@ -33,6 +33,17 @@ namespace UE::UAF::Tests
 		return SetMap == Other;
 	}
 
+	static bool CollectionContains(const FAttributeCollection& Collection, const FAttributeMap* Map)
+	{
+		if (Map == nullptr)
+		{
+			return false;
+		}
+
+		const FAttributeMap* Other = Collection.FindMap(Map->GetValueType());
+		return Map == Other;
+	}
+
 	IMPLEMENT_SIMPLE_AUTOMATION_TEST(FIteratorUtilsInnerJoinTest, "Animation.UAF.Attributes.IteratorUtils.InnerJoin", EAutomationTestFlags::EditorContext | EAutomationTestFlags::EngineFilter)
 
 	bool FIteratorUtilsInnerJoinTest::RunTest(const FString& InParameters)
@@ -127,8 +138,15 @@ namespace UE::UAF::Tests
 			TAttributeSetMap<FFloatAnimationAttribute>* TransformToFloatMap0 = Collection0.AddSetMap<FFloatAnimationAttribute>(TransformToFloatKey);
 			UE_RETURN_ON_ERROR(TransformToFloatMap0 != nullptr, TEXT("Transform->Bone map should be valid"));
 
+			TAttributeMap<FFloatAnimationAttribute>* FloatMap0 = Collection0.AddMap<FFloatAnimationAttribute>();
+			UE_RETURN_ON_ERROR(FloatMap0 != nullptr, TEXT("Float map should be valid"));
+
+			TAttributeMap<FBoneTransformAnimationAttribute>* BoneMap0 = Collection0.AddMap<FBoneTransformAnimationAttribute>();
+			UE_RETURN_ON_ERROR(BoneMap0 != nullptr, TEXT("Bone map should be valid"));
+
 			AddErrorIfFalse(!Collection0.IsEmpty(), TEXT("Attribute collection should not be empty"));
 			AddErrorIfFalse(Collection0.NumSetMaps() == 5, TEXT("Unexpected set map count"));
+			AddErrorIfFalse(Collection0.NumMaps() == 2, TEXT("Unexpected set count"));
 
 			FAttributeCollectionHeap Collection1(FullBodySet);
 
@@ -150,10 +168,17 @@ namespace UE::UAF::Tests
 			TAttributeSetMap<FFloatAnimationAttribute>* TransformToFloatMap1 = Collection1.AddSetMap<FFloatAnimationAttribute>(TransformToFloatKey);
 			UE_RETURN_ON_ERROR(TransformToFloatMap1 != nullptr, TEXT("Transform->Bone map should be valid"));
 
+			TAttributeMap<FFloatAnimationAttribute>* FloatMap1 = Collection1.AddMap<FFloatAnimationAttribute>();
+			UE_RETURN_ON_ERROR(FloatMap1 != nullptr, TEXT("Float map should be valid"));
+
+			TAttributeMap<FBoneTransformAnimationAttribute>* BoneMap1 = Collection1.AddMap<FBoneTransformAnimationAttribute>();
+			UE_RETURN_ON_ERROR(BoneMap1 != nullptr, TEXT("Bone map should be valid"));
+
 			AddErrorIfFalse(!Collection1.IsEmpty(), TEXT("Attribute collection should not be empty"));
 			AddErrorIfFalse(Collection1.NumSetMaps() == 5, TEXT("Unexpected set map count"));
+			AddErrorIfFalse(Collection1.NumMaps() == 2, TEXT("Unexpected set count"));
 
-			// Both inputs perfectly match
+			// Both inputs perfectly match (set maps)
 			{
 				int32 NumPredicateCalls = 0;
 				InnerJoinBy(
@@ -216,6 +241,65 @@ namespace UE::UAF::Tests
 				AddErrorIfFalse(NumPredicateCalls == 5, TEXT("Unexpected InnerJoin iteration result"));
 			}
 
+			// Both inputs perfectly match (name/value maps)
+			{
+				int32 NumPredicateCalls = 0;
+				InnerJoinBy(
+					FAttributeCollectionMapJoinOp(),
+					[this, &NumPredicateCalls, &Collection0, &Collection1](FAttributeMap* Map0, FAttributeMap* Map1)
+					{
+						NumPredicateCalls++;
+						AddErrorIfFalse(CollectionContains(Collection0, Map0), TEXT("Unexpected InnerJoin iteration result"));
+						AddErrorIfFalse(CollectionContains(Collection1, Map1), TEXT("Unexpected InnerJoin iteration result"));
+						AddErrorIfFalse(Map0 != nullptr && Map1 != nullptr && Map0->GetValueType() == Map1->GetValueType(), TEXT("Unexpected InnerJoin iteration result"));
+					},
+					Collection0.CreateMapIterator(),
+					Collection1.CreateMapIterator());
+				AddErrorIfFalse(NumPredicateCalls == 2, TEXT("Unexpected InnerJoin iteration result"));
+
+				NumPredicateCalls = 0;
+				InnerJoinBy(
+					FAttributeCollectionMapJoinOp(),
+					[this, &NumPredicateCalls, &Collection0, &Collection1](const FAttributeMap* Map0, const FAttributeMap* Map1)
+					{
+						NumPredicateCalls++;
+						AddErrorIfFalse(CollectionContains(Collection0, Map0), TEXT("Unexpected InnerJoin iteration result"));
+						AddErrorIfFalse(CollectionContains(Collection1, Map1), TEXT("Unexpected InnerJoin iteration result"));
+						AddErrorIfFalse(Map0 != nullptr && Map1 != nullptr && Map0->GetValueType() == Map1->GetValueType(), TEXT("Unexpected InnerJoin iteration result"));
+					},
+					Collection0.CreateConstMapIterator(),
+					Collection1.CreateConstMapIterator());
+				AddErrorIfFalse(NumPredicateCalls == 2, TEXT("Unexpected InnerJoin iteration result"));
+
+				NumPredicateCalls = 0;
+				InnerJoinBy(
+					FAttributeCollectionMapJoinOp(),
+					[this, &NumPredicateCalls, &Collection0, &Collection1](const FAttributeMap* Map0, FAttributeMap* Map1)
+					{
+						NumPredicateCalls++;
+						AddErrorIfFalse(CollectionContains(Collection0, Map0), TEXT("Unexpected InnerJoin iteration result"));
+						AddErrorIfFalse(CollectionContains(Collection1, Map1), TEXT("Unexpected InnerJoin iteration result"));
+						AddErrorIfFalse(Map0 != nullptr && Map1 != nullptr && Map0->GetValueType() == Map1->GetValueType(), TEXT("Unexpected InnerJoin iteration result"));
+					},
+					Collection0.CreateConstMapIterator(),
+					Collection1.CreateMapIterator());
+				AddErrorIfFalse(NumPredicateCalls == 2, TEXT("Unexpected InnerJoin iteration result"));
+
+				NumPredicateCalls = 0;
+				InnerJoinBy(
+					FAttributeCollectionMapJoinOp(),
+					[this, &NumPredicateCalls, &Collection0, &Collection1](FAttributeMap* Map0, const FAttributeMap* Map1)
+					{
+						NumPredicateCalls++;
+						AddErrorIfFalse(CollectionContains(Collection0, Map0), TEXT("Unexpected InnerJoin iteration result"));
+						AddErrorIfFalse(CollectionContains(Collection1, Map1), TEXT("Unexpected InnerJoin iteration result"));
+						AddErrorIfFalse(Map0 != nullptr && Map1 != nullptr && Map0->GetValueType() == Map1->GetValueType(), TEXT("Unexpected InnerJoin iteration result"));
+					},
+					Collection0.CreateMapIterator(),
+					Collection1.CreateConstMapIterator());
+				AddErrorIfFalse(NumPredicateCalls == 2, TEXT("Unexpected InnerJoin iteration result"));
+			}
+
 			FAttributeCollectionHeap Collection2(FullBodySet);
 
 			AddErrorIfFalse(Collection2.IsEmpty(), TEXT("Attribute collection should be empty"));
@@ -227,10 +311,14 @@ namespace UE::UAF::Tests
 			TAttributeSetMap<FBoneTransformAnimationAttribute>* TransformToTransformMap2 = Collection2.AddSetMap<FBoneTransformAnimationAttribute>(TransformToTransformKey);
 			UE_RETURN_ON_ERROR(TransformToTransformMap2 != nullptr, TEXT("Transform->Transform map should be valid"));
 
+			TAttributeMap<FFloatAnimationAttribute>* FloatMap2 = Collection2.AddMap<FFloatAnimationAttribute>();
+			UE_RETURN_ON_ERROR(FloatMap2 != nullptr, TEXT("Float map should be valid"));
+
 			AddErrorIfFalse(!Collection2.IsEmpty(), TEXT("Attribute collection should not be empty"));
 			AddErrorIfFalse(Collection2.NumSetMaps() == 2, TEXT("Unexpected set map count"));
+			AddErrorIfFalse(Collection2.NumMaps() == 1, TEXT("Unexpected map count"));
 
-			// First input is shorter and its entries are skipped when missing
+			// First input is shorter and its entries are skipped when missing (set maps)
 			{
 				int32 NumPredicateCalls = 0;
 				InnerJoinBy(
@@ -293,7 +381,66 @@ namespace UE::UAF::Tests
 				AddErrorIfFalse(NumPredicateCalls == 2, TEXT("Unexpected InnerJoin iteration result"));
 			}
 
-			// Second input is shorter and its entries are skipped when missing
+			// First input is shorter and its entries are skipped when missing (name/value maps)
+			{
+				int32 NumPredicateCalls = 0;
+				InnerJoinBy(
+					FAttributeCollectionMapJoinOp(),
+					[this, &NumPredicateCalls, &Collection0, &Collection2](FAttributeMap* Map2, FAttributeMap* Map0)
+					{
+						NumPredicateCalls++;
+						AddErrorIfFalse(CollectionContains(Collection0, Map0), TEXT("Unexpected InnerJoin iteration result"));
+						AddErrorIfFalse(CollectionContains(Collection2, Map2), TEXT("Unexpected InnerJoin iteration result"));
+						AddErrorIfFalse(Map0 != nullptr && Map2 != nullptr && Map0->GetValueType() == Map2->GetValueType(), TEXT("Unexpected InnerJoin iteration result"));
+					},
+					Collection2.CreateMapIterator(),
+					Collection0.CreateMapIterator());
+				AddErrorIfFalse(NumPredicateCalls == 1, TEXT("Unexpected InnerJoin iteration result"));
+
+				NumPredicateCalls = 0;
+				InnerJoinBy(
+					FAttributeCollectionMapJoinOp(),
+					[this, &NumPredicateCalls, &Collection0, &Collection2](const FAttributeMap* Map2, const FAttributeMap* Map0)
+					{
+						NumPredicateCalls++;
+						AddErrorIfFalse(CollectionContains(Collection0, Map0), TEXT("Unexpected InnerJoin iteration result"));
+						AddErrorIfFalse(CollectionContains(Collection2, Map2), TEXT("Unexpected InnerJoin iteration result"));
+						AddErrorIfFalse(Map0 != nullptr && Map2 != nullptr && Map0->GetValueType() == Map2->GetValueType(), TEXT("Unexpected InnerJoin iteration result"));
+					},
+					Collection2.CreateConstMapIterator(),
+					Collection0.CreateConstMapIterator());
+				AddErrorIfFalse(NumPredicateCalls == 1, TEXT("Unexpected InnerJoin iteration result"));
+
+				NumPredicateCalls = 0;
+				InnerJoinBy(
+					FAttributeCollectionMapJoinOp(),
+					[this, &NumPredicateCalls, &Collection0, &Collection2](const FAttributeMap* Map2, FAttributeMap* Map0)
+					{
+						NumPredicateCalls++;
+						AddErrorIfFalse(CollectionContains(Collection0, Map0), TEXT("Unexpected InnerJoin iteration result"));
+						AddErrorIfFalse(CollectionContains(Collection2, Map2), TEXT("Unexpected InnerJoin iteration result"));
+						AddErrorIfFalse(Map0 != nullptr && Map2 != nullptr && Map0->GetValueType() == Map2->GetValueType(), TEXT("Unexpected InnerJoin iteration result"));
+					},
+					Collection2.CreateConstMapIterator(),
+					Collection0.CreateMapIterator());
+				AddErrorIfFalse(NumPredicateCalls == 1, TEXT("Unexpected InnerJoin iteration result"));
+
+				NumPredicateCalls = 0;
+				InnerJoinBy(
+					FAttributeCollectionMapJoinOp(),
+					[this, &NumPredicateCalls, &Collection0, &Collection2](FAttributeMap* Map2, const FAttributeMap* Map0)
+					{
+						NumPredicateCalls++;
+						AddErrorIfFalse(CollectionContains(Collection0, Map0), TEXT("Unexpected InnerJoin iteration result"));
+						AddErrorIfFalse(CollectionContains(Collection2, Map2), TEXT("Unexpected InnerJoin iteration result"));
+						AddErrorIfFalse(Map0 != nullptr && Map2 != nullptr && Map0->GetValueType() == Map2->GetValueType(), TEXT("Unexpected InnerJoin iteration result"));
+					},
+					Collection2.CreateMapIterator(),
+					Collection0.CreateConstMapIterator());
+				AddErrorIfFalse(NumPredicateCalls == 1, TEXT("Unexpected InnerJoin iteration result"));
+			}
+
+			// Second input is shorter and its entries are skipped when missing (set maps)
 			{
 				int32 NumPredicateCalls = 0;
 				InnerJoinBy(
@@ -356,7 +503,66 @@ namespace UE::UAF::Tests
 				AddErrorIfFalse(NumPredicateCalls == 2, TEXT("Unexpected InnerJoin iteration result"));
 			}
 
-			// 3 input variant
+			// Second input is shorter and its entries are skipped when missing (name/value maps)
+			{
+				int32 NumPredicateCalls = 0;
+				InnerJoinBy(
+					FAttributeCollectionMapJoinOp(),
+					[this, &NumPredicateCalls, &Collection0, &Collection2](FAttributeMap* Map0, FAttributeMap* Map2)
+					{
+						NumPredicateCalls++;
+						AddErrorIfFalse(CollectionContains(Collection0, Map0), TEXT("Unexpected InnerJoin iteration result"));
+						AddErrorIfFalse(CollectionContains(Collection2, Map2), TEXT("Unexpected InnerJoin iteration result"));
+						AddErrorIfFalse(Map0 != nullptr && Map2 != nullptr && Map0->GetValueType() == Map2->GetValueType(), TEXT("Unexpected InnerJoin iteration result"));
+					},
+					Collection0.CreateMapIterator(),
+					Collection2.CreateMapIterator());
+				AddErrorIfFalse(NumPredicateCalls == 1, TEXT("Unexpected InnerJoin iteration result"));
+
+				NumPredicateCalls = 0;
+				InnerJoinBy(
+					FAttributeCollectionMapJoinOp(),
+					[this, &NumPredicateCalls, &Collection0, &Collection2](const FAttributeMap* Map0, const FAttributeMap* Map2)
+					{
+						NumPredicateCalls++;
+						AddErrorIfFalse(CollectionContains(Collection0, Map0), TEXT("Unexpected InnerJoin iteration result"));
+						AddErrorIfFalse(CollectionContains(Collection2, Map2), TEXT("Unexpected InnerJoin iteration result"));
+						AddErrorIfFalse(Map0 != nullptr && Map2 != nullptr && Map0->GetValueType() == Map2->GetValueType(), TEXT("Unexpected InnerJoin iteration result"));
+					},
+					Collection0.CreateConstMapIterator(),
+					Collection2.CreateConstMapIterator());
+				AddErrorIfFalse(NumPredicateCalls == 1, TEXT("Unexpected InnerJoin iteration result"));
+
+				NumPredicateCalls = 0;
+				InnerJoinBy(
+					FAttributeCollectionMapJoinOp(),
+					[this, &NumPredicateCalls, &Collection0, &Collection2](const FAttributeMap* Map0, FAttributeMap* Map2)
+					{
+						NumPredicateCalls++;
+						AddErrorIfFalse(CollectionContains(Collection0, Map0), TEXT("Unexpected InnerJoin iteration result"));
+						AddErrorIfFalse(CollectionContains(Collection2, Map2), TEXT("Unexpected InnerJoin iteration result"));
+						AddErrorIfFalse(Map0 != nullptr && Map2 != nullptr && Map0->GetValueType() == Map2->GetValueType(), TEXT("Unexpected InnerJoin iteration result"));
+					},
+					Collection0.CreateConstMapIterator(),
+					Collection2.CreateMapIterator());
+				AddErrorIfFalse(NumPredicateCalls == 1, TEXT("Unexpected InnerJoin iteration result"));
+
+				NumPredicateCalls = 0;
+				InnerJoinBy(
+					FAttributeCollectionMapJoinOp(),
+					[this, &NumPredicateCalls, &Collection0, &Collection2](FAttributeMap* Map0, const FAttributeMap* Map2)
+					{
+						NumPredicateCalls++;
+						AddErrorIfFalse(CollectionContains(Collection0, Map0), TEXT("Unexpected InnerJoin iteration result"));
+						AddErrorIfFalse(CollectionContains(Collection2, Map2), TEXT("Unexpected InnerJoin iteration result"));
+						AddErrorIfFalse(Map0 != nullptr && Map2 != nullptr && Map0->GetValueType() == Map2->GetValueType(), TEXT("Unexpected InnerJoin iteration result"));
+					},
+					Collection0.CreateMapIterator(),
+					Collection2.CreateConstMapIterator());
+				AddErrorIfFalse(NumPredicateCalls == 1, TEXT("Unexpected InnerJoin iteration result"));
+			}
+
+			// 3 input variant (set maps)
 			{
 				int32 NumPredicateCalls = 0;
 				InnerJoinBy(
@@ -473,7 +679,112 @@ namespace UE::UAF::Tests
 				AddErrorIfFalse(NumPredicateCalls == 2, TEXT("Unexpected InnerJoin iteration result"));
 			}
 
-			// Join with a transformer list
+			// 3 input variant (name/value maps)
+			{
+				int32 NumPredicateCalls = 0;
+				InnerJoinBy(
+					FAttributeCollectionMapJoinOp(),
+					[this, &NumPredicateCalls, &Collection0, &Collection1, &Collection2](FAttributeMap* Map0, FAttributeMap* Map1, FAttributeMap* Map2)
+					{
+						NumPredicateCalls++;
+						AddErrorIfFalse(CollectionContains(Collection0, Map0), TEXT("Unexpected InnerJoin iteration result"));
+						AddErrorIfFalse(CollectionContains(Collection1, Map1), TEXT("Unexpected InnerJoin iteration result"));
+						AddErrorIfFalse(CollectionContains(Collection2, Map2), TEXT("Unexpected InnerJoin iteration result"));
+						AddErrorIfFalse(Map0 != nullptr && Map1 != nullptr && Map0->GetValueType() == Map1->GetValueType(), TEXT("Unexpected InnerJoin iteration result"));
+						AddErrorIfFalse(Map0 != nullptr && Map2 != nullptr && Map0->GetValueType() == Map2->GetValueType(), TEXT("Unexpected InnerJoin iteration result"));
+					},
+					Collection0.CreateMapIterator(),
+					Collection1.CreateMapIterator(),
+					Collection2.CreateMapIterator());
+				AddErrorIfFalse(NumPredicateCalls == 1, TEXT("Unexpected InnerJoin iteration result"));
+
+				NumPredicateCalls = 0;
+				InnerJoinBy(
+					FAttributeCollectionMapJoinOp(),
+					[this, &NumPredicateCalls, &Collection0, &Collection1, &Collection2](FAttributeMap* Map0, FAttributeMap* Map2, FAttributeMap* Map1)
+					{
+						NumPredicateCalls++;
+						AddErrorIfFalse(CollectionContains(Collection0, Map0), TEXT("Unexpected InnerJoin iteration result"));
+						AddErrorIfFalse(CollectionContains(Collection1, Map1), TEXT("Unexpected InnerJoin iteration result"));
+						AddErrorIfFalse(CollectionContains(Collection2, Map2), TEXT("Unexpected InnerJoin iteration result"));
+						AddErrorIfFalse(Map0 != nullptr && Map1 != nullptr && Map0->GetValueType() == Map1->GetValueType(), TEXT("Unexpected InnerJoin iteration result"));
+						AddErrorIfFalse(Map0 != nullptr && Map2 != nullptr && Map0->GetValueType() == Map2->GetValueType(), TEXT("Unexpected InnerJoin iteration result"));
+					},
+					Collection0.CreateMapIterator(),
+					Collection2.CreateMapIterator(),
+					Collection1.CreateMapIterator());
+				AddErrorIfFalse(NumPredicateCalls == 1, TEXT("Unexpected InnerJoin iteration result"));
+
+				NumPredicateCalls = 0;
+				InnerJoinBy(
+					FAttributeCollectionMapJoinOp(),
+					[this, &NumPredicateCalls, &Collection0, &Collection1, &Collection2](FAttributeMap* Map1, FAttributeMap* Map0, FAttributeMap* Map2)
+					{
+						NumPredicateCalls++;
+						AddErrorIfFalse(CollectionContains(Collection0, Map0), TEXT("Unexpected InnerJoin iteration result"));
+						AddErrorIfFalse(CollectionContains(Collection1, Map1), TEXT("Unexpected InnerJoin iteration result"));
+						AddErrorIfFalse(CollectionContains(Collection2, Map2), TEXT("Unexpected InnerJoin iteration result"));
+						AddErrorIfFalse(Map0 != nullptr && Map1 != nullptr && Map0->GetValueType() == Map1->GetValueType(), TEXT("Unexpected InnerJoin iteration result"));
+						AddErrorIfFalse(Map0 != nullptr && Map2 != nullptr && Map0->GetValueType() == Map2->GetValueType(), TEXT("Unexpected InnerJoin iteration result"));
+					},
+					Collection1.CreateMapIterator(),
+					Collection0.CreateMapIterator(),
+					Collection2.CreateMapIterator());
+				AddErrorIfFalse(NumPredicateCalls == 1, TEXT("Unexpected InnerJoin iteration result"));
+
+				NumPredicateCalls = 0;
+				InnerJoinBy(
+					FAttributeCollectionMapJoinOp(),
+					[this, &NumPredicateCalls, &Collection0, &Collection1, &Collection2](FAttributeMap* Map1, FAttributeMap* Map2, FAttributeMap* Map0)
+					{
+						NumPredicateCalls++;
+						AddErrorIfFalse(CollectionContains(Collection0, Map0), TEXT("Unexpected InnerJoin iteration result"));
+						AddErrorIfFalse(CollectionContains(Collection1, Map1), TEXT("Unexpected InnerJoin iteration result"));
+						AddErrorIfFalse(CollectionContains(Collection2, Map2), TEXT("Unexpected InnerJoin iteration result"));
+						AddErrorIfFalse(Map0 != nullptr && Map1 != nullptr && Map0->GetValueType() == Map1->GetValueType(), TEXT("Unexpected InnerJoin iteration result"));
+						AddErrorIfFalse(Map0 != nullptr && Map2 != nullptr && Map0->GetValueType() == Map2->GetValueType(), TEXT("Unexpected InnerJoin iteration result"));
+					},
+					Collection1.CreateMapIterator(),
+					Collection2.CreateMapIterator(),
+					Collection0.CreateMapIterator());
+				AddErrorIfFalse(NumPredicateCalls == 1, TEXT("Unexpected InnerJoin iteration result"));
+
+				NumPredicateCalls = 0;
+				InnerJoinBy(
+					FAttributeCollectionMapJoinOp(),
+					[this, &NumPredicateCalls, &Collection0, &Collection1, &Collection2](FAttributeMap* Map2, FAttributeMap* Map0, FAttributeMap* Map1)
+					{
+						NumPredicateCalls++;
+						AddErrorIfFalse(CollectionContains(Collection0, Map0), TEXT("Unexpected InnerJoin iteration result"));
+						AddErrorIfFalse(CollectionContains(Collection1, Map1), TEXT("Unexpected InnerJoin iteration result"));
+						AddErrorIfFalse(CollectionContains(Collection2, Map2), TEXT("Unexpected InnerJoin iteration result"));
+						AddErrorIfFalse(Map0 != nullptr && Map1 != nullptr && Map0->GetValueType() == Map1->GetValueType(), TEXT("Unexpected InnerJoin iteration result"));
+						AddErrorIfFalse(Map0 != nullptr && Map2 != nullptr && Map0->GetValueType() == Map2->GetValueType(), TEXT("Unexpected InnerJoin iteration result"));
+					},
+					Collection2.CreateMapIterator(),
+					Collection0.CreateMapIterator(),
+					Collection1.CreateMapIterator());
+				AddErrorIfFalse(NumPredicateCalls == 1, TEXT("Unexpected InnerJoin iteration result"));
+
+				NumPredicateCalls = 0;
+				InnerJoinBy(
+					FAttributeCollectionMapJoinOp(),
+					[this, &NumPredicateCalls, &Collection0, &Collection1, &Collection2](FAttributeMap* Map2, FAttributeMap* Map1, FAttributeMap* Map0)
+					{
+						NumPredicateCalls++;
+						AddErrorIfFalse(CollectionContains(Collection0, Map0), TEXT("Unexpected InnerJoin iteration result"));
+						AddErrorIfFalse(CollectionContains(Collection1, Map1), TEXT("Unexpected InnerJoin iteration result"));
+						AddErrorIfFalse(CollectionContains(Collection2, Map2), TEXT("Unexpected InnerJoin iteration result"));
+						AddErrorIfFalse(Map0 != nullptr && Map1 != nullptr && Map0->GetValueType() == Map1->GetValueType(), TEXT("Unexpected InnerJoin iteration result"));
+						AddErrorIfFalse(Map0 != nullptr && Map2 != nullptr && Map0->GetValueType() == Map2->GetValueType(), TEXT("Unexpected InnerJoin iteration result"));
+					},
+					Collection2.CreateMapIterator(),
+					Collection1.CreateMapIterator(),
+					Collection0.CreateMapIterator());
+				AddErrorIfFalse(NumPredicateCalls == 1, TEXT("Unexpected InnerJoin iteration result"));
+			}
+
+			// Join with a transformer list (set maps)
 			{
 				// Use a float transformer
 				// We set up Collection0 such that it has two mappings with float values: Transform -> Float and Float -> Float
@@ -485,7 +796,7 @@ namespace UE::UAF::Tests
 					using TransformerType = Transformers::FSanitize;
 					using ValueType = FFloatAnimationAttribute;
 
-					static void Transform(FAttributeSetMap* SetMap) {}	// No-op
+					static void TransformSetMap(FAttributeSetMap* SetMap) {}	// No-op
 				};
 
 				// Dummy vector sanitize transformer
@@ -494,7 +805,7 @@ namespace UE::UAF::Tests
 					using TransformerType = Transformers::FSanitize;
 					using ValueType = FVectorAnimationAttribute;
 
-					static void Transform(FAttributeSetMap* SetMap) {}	// No-op
+					static void TransformSetMap(FAttributeSetMap* SetMap) {}	// No-op
 				};
 
 #ifdef _MSC_VER
@@ -503,8 +814,8 @@ namespace UE::UAF::Tests
 #endif
 
 				FAttributeValueTransformerList TransformerList(Transformers::FSanitize::TransformerName);
-				TransformerList.AddTransformer(FFloatAnimationAttribute::StaticStruct(), reinterpret_cast<FRawTransformerFunc>(&FFloatAttributeTransformer_SanitizeTest::Transform));
-				TransformerList.AddTransformer(FVectorAnimationAttribute::StaticStruct(), reinterpret_cast<FRawTransformerFunc>(&FVectorAttributeTransformer_SanitizeTest::Transform));
+				TransformerList.AddSetMapTransformer(FFloatAnimationAttribute::StaticStruct(), reinterpret_cast<FRawTransformerFunc>(&FFloatAttributeTransformer_SanitizeTest::TransformSetMap));
+				TransformerList.AddSetMapTransformer(FVectorAnimationAttribute::StaticStruct(), reinterpret_cast<FRawTransformerFunc>(&FVectorAttributeTransformer_SanitizeTest::TransformSetMap));
 
 #ifdef _MSC_VER
 	#pragma warning(pop)
@@ -512,17 +823,68 @@ namespace UE::UAF::Tests
 
 				int32 NumPredicateCalls = 0;
 				InnerJoinBy(
-					FAttributeValueTransformerListWithAttributeCollectionJoinOp(),
+					FAttributeValueTransformerListWithAttributeCollectionSetMapJoinOp(),
 					[this, &NumPredicateCalls, &Collection0](FRawTransformerFunc Transformer, FAttributeSetMap* SetMap)
 					{
 						NumPredicateCalls++;
 						AddErrorIfFalse(CollectionContains(Collection0, SetMap), TEXT("Unexpected InnerJoin iteration result"));
 						AddErrorIfFalse(SetMap->GetValueType() == FFloatAnimationAttribute::StaticStruct(), TEXT("Unexpected InnerJoin iteration result"));
 					},
-					TransformerList.CreateTransformerIterator(),
+					TransformerList.CreateSetMapTransformerIterator(),
 					Collection0.CreateSetMapIterator());
 				AddErrorIfFalse(NumPredicateCalls == 2, TEXT("Unexpected InnerJoin iteration result"));
 			}
+
+			// Join with a transformer list (name/value maps)
+			{
+				// Use a float transformer
+				// We set up Collection0 such that it has two mappings with these value types: Float and Bone Transform
+				// InnerJoin should execute the predicate twice, once for each value type that matches the transformer's value type
+
+				// Dummy float sanitize transformer
+				struct FFloatAttributeTransformer_SanitizeTest
+				{
+					using TransformerType = Transformers::FSanitize;
+					using ValueType = FFloatAnimationAttribute;
+
+					static void TransformMap(FAttributeMap* Map) {}	// No-op
+				};
+
+				// Dummy vector sanitize transformer
+				struct FVectorAttributeTransformer_SanitizeTest
+				{
+					using TransformerType = Transformers::FSanitize;
+					using ValueType = FVectorAnimationAttribute;
+
+					static void TransformMap(FAttributeMap* Map) {}	// No-op
+				};
+
+#ifdef _MSC_VER
+	#pragma warning(push)
+	#pragma warning(disable : 4191) // unsafe conversion from 'function pointer type A' to 'function pointer type B'
+#endif
+
+				FAttributeValueTransformerList TransformerList(Transformers::FSanitize::TransformerName);
+				TransformerList.AddMapTransformer(FFloatAnimationAttribute::StaticStruct(), reinterpret_cast<FRawTransformerFunc>(&FFloatAttributeTransformer_SanitizeTest::TransformMap));
+				TransformerList.AddMapTransformer(FVectorAnimationAttribute::StaticStruct(), reinterpret_cast<FRawTransformerFunc>(&FVectorAttributeTransformer_SanitizeTest::TransformMap));
+
+#ifdef _MSC_VER
+	#pragma warning(pop)
+#endif
+
+				int32 NumPredicateCalls = 0;
+				InnerJoinBy(
+					FAttributeValueTransformerListWithAttributeCollectionMapJoinOp(),
+					[this, &NumPredicateCalls, &Collection0](FRawTransformerFunc Transformer, FAttributeMap* Map)
+					{
+						NumPredicateCalls++;
+						AddErrorIfFalse(CollectionContains(Collection0, Map), TEXT("Unexpected InnerJoin iteration result"));
+						AddErrorIfFalse(Map->GetValueType() == FFloatAnimationAttribute::StaticStruct(), TEXT("Unexpected InnerJoin iteration result"));
+					},
+					TransformerList.CreateMapTransformerIterator(),
+					Collection0.CreateMapIterator());
+				AddErrorIfFalse(NumPredicateCalls == 1, TEXT("Unexpected InnerJoin iteration result"));
+			}
 		}
 
 		return true;
@@ -622,8 +984,15 @@ namespace UE::UAF::Tests
 			TAttributeSetMap<FFloatAnimationAttribute>* TransformToFloatMap0 = Collection0.AddSetMap<FFloatAnimationAttribute>(TransformToFloatKey);
 			UE_RETURN_ON_ERROR(TransformToFloatMap0 != nullptr, TEXT("Transform->Bone map should be valid"));
 
+			TAttributeMap<FFloatAnimationAttribute>* FloatMap0 = Collection0.AddMap<FFloatAnimationAttribute>();
+			UE_RETURN_ON_ERROR(FloatMap0 != nullptr, TEXT("Float map should be valid"));
+
+			TAttributeMap<FBoneTransformAnimationAttribute>* BoneMap0 = Collection0.AddMap<FBoneTransformAnimationAttribute>();
+			UE_RETURN_ON_ERROR(BoneMap0 != nullptr, TEXT("Bone map should be valid"));
+
 			AddErrorIfFalse(!Collection0.IsEmpty(), TEXT("Attribute collection should not be empty"));
 			AddErrorIfFalse(Collection0.NumSetMaps() == 5, TEXT("Unexpected set map count"));
+			AddErrorIfFalse(Collection0.NumMaps() == 2, TEXT("Unexpected map count"));
 
 			FAttributeCollectionHeap Collection1(FullBodySet);
 
@@ -645,10 +1014,17 @@ namespace UE::UAF::Tests
 			TAttributeSetMap<FFloatAnimationAttribute>* TransformToFloatMap1 = Collection1.AddSetMap<FFloatAnimationAttribute>(TransformToFloatKey);
 			UE_RETURN_ON_ERROR(TransformToFloatMap1 != nullptr, TEXT("Transform->Bone map should be valid"));
 
+			TAttributeMap<FFloatAnimationAttribute>* FloatMap1 = Collection1.AddMap<FFloatAnimationAttribute>();
+			UE_RETURN_ON_ERROR(FloatMap1 != nullptr, TEXT("Float map should be valid"));
+
+			TAttributeMap<FBoneTransformAnimationAttribute>* BoneMap1 = Collection1.AddMap<FBoneTransformAnimationAttribute>();
+			UE_RETURN_ON_ERROR(BoneMap1 != nullptr, TEXT("Bone map should be valid"));
+
 			AddErrorIfFalse(!Collection1.IsEmpty(), TEXT("Attribute collection should not be empty"));
 			AddErrorIfFalse(Collection1.NumSetMaps() == 5, TEXT("Unexpected set map count"));
+			AddErrorIfFalse(Collection1.NumMaps() == 2, TEXT("Unexpected map count"));
 
-			// Both inputs perfectly match
+			// Both inputs perfectly match (set maps)
 			{
 				int32 NumPredicateCalls = 0;
 				OuterJoinBy(
@@ -711,6 +1087,65 @@ namespace UE::UAF::Tests
 				AddErrorIfFalse(NumPredicateCalls == 5, TEXT("Unexpected OuterJoin iteration result"));
 			}
 
+			// Both inputs perfectly match (name/value maps)
+			{
+				int32 NumPredicateCalls = 0;
+				OuterJoinBy(
+					FAttributeCollectionMapJoinOp(),
+					[this, &NumPredicateCalls, &Collection0, &Collection1](FAttributeMap* Map0, FAttributeMap* Map1)
+					{
+						NumPredicateCalls++;
+						AddErrorIfFalse(CollectionContains(Collection0, Map0), TEXT("Unexpected OuterJoin iteration result"));
+						AddErrorIfFalse(CollectionContains(Collection1, Map1), TEXT("Unexpected OuterJoin iteration result"));
+						AddErrorIfFalse(Map0 != nullptr && Map1 != nullptr && Map0->GetValueType() == Map1->GetValueType(), TEXT("Unexpected OuterJoin iteration result"));
+					},
+					Collection0.CreateMapIterator(),
+					Collection1.CreateMapIterator());
+				AddErrorIfFalse(NumPredicateCalls == 2, TEXT("Unexpected OuterJoin iteration result"));
+
+				NumPredicateCalls = 0;
+				OuterJoinBy(
+					FAttributeCollectionMapJoinOp(),
+					[this, &NumPredicateCalls, &Collection0, &Collection1](const FAttributeMap* Map0, const FAttributeMap* Map1)
+					{
+						NumPredicateCalls++;
+						AddErrorIfFalse(CollectionContains(Collection0, Map0), TEXT("Unexpected OuterJoin iteration result"));
+						AddErrorIfFalse(CollectionContains(Collection1, Map1), TEXT("Unexpected OuterJoin iteration result"));
+						AddErrorIfFalse(Map0 != nullptr && Map1 != nullptr && Map0->GetValueType() == Map1->GetValueType(), TEXT("Unexpected OuterJoin iteration result"));
+					},
+					Collection0.CreateConstMapIterator(),
+					Collection1.CreateConstMapIterator());
+				AddErrorIfFalse(NumPredicateCalls == 2, TEXT("Unexpected OuterJoin iteration result"));
+
+				NumPredicateCalls = 0;
+				OuterJoinBy(
+					FAttributeCollectionMapJoinOp(),
+					[this, &NumPredicateCalls, &Collection0, &Collection1](const FAttributeMap* Map0, FAttributeMap* Map1)
+					{
+						NumPredicateCalls++;
+						AddErrorIfFalse(CollectionContains(Collection0, Map0), TEXT("Unexpected OuterJoin iteration result"));
+						AddErrorIfFalse(CollectionContains(Collection1, Map1), TEXT("Unexpected OuterJoin iteration result"));
+						AddErrorIfFalse(Map0 != nullptr && Map1 != nullptr && Map0->GetValueType() == Map1->GetValueType(), TEXT("Unexpected OuterJoin iteration result"));
+					},
+					Collection0.CreateConstMapIterator(),
+					Collection1.CreateMapIterator());
+				AddErrorIfFalse(NumPredicateCalls == 2, TEXT("Unexpected OuterJoin iteration result"));
+
+				NumPredicateCalls = 0;
+				OuterJoinBy(
+					FAttributeCollectionMapJoinOp(),
+					[this, &NumPredicateCalls, &Collection0, &Collection1](FAttributeMap* Map0, const FAttributeMap* Map1)
+					{
+						NumPredicateCalls++;
+						AddErrorIfFalse(CollectionContains(Collection0, Map0), TEXT("Unexpected OuterJoin iteration result"));
+						AddErrorIfFalse(CollectionContains(Collection1, Map1), TEXT("Unexpected OuterJoin iteration result"));
+						AddErrorIfFalse(Map0 != nullptr && Map1 != nullptr && Map0->GetValueType() == Map1->GetValueType(), TEXT("Unexpected OuterJoin iteration result"));
+					},
+					Collection0.CreateMapIterator(),
+					Collection1.CreateConstMapIterator());
+				AddErrorIfFalse(NumPredicateCalls == 2, TEXT("Unexpected OuterJoin iteration result"));
+			}
+
 			FAttributeCollectionHeap Collection2(FullBodySet);
 
 			AddErrorIfFalse(Collection2.IsEmpty(), TEXT("Attribute collection should be empty"));
@@ -722,10 +1157,14 @@ namespace UE::UAF::Tests
 			TAttributeSetMap<FBoneTransformAnimationAttribute>* TransformToTransformMap2 = Collection2.AddSetMap<FBoneTransformAnimationAttribute>(TransformToTransformKey);
 			UE_RETURN_ON_ERROR(TransformToTransformMap2 != nullptr, TEXT("Transform->Transform map should be valid"));
 
+			TAttributeMap<FFloatAnimationAttribute>* FloatMap2 = Collection2.AddMap<FFloatAnimationAttribute>();
+			UE_RETURN_ON_ERROR(FloatMap2 != nullptr, TEXT("Float map should be valid"));
+
 			AddErrorIfFalse(!Collection2.IsEmpty(), TEXT("Attribute collection should not be empty"));
 			AddErrorIfFalse(Collection2.NumSetMaps() == 2, TEXT("Unexpected set map count"));
+			AddErrorIfFalse(Collection2.NumMaps() == 1, TEXT("Unexpected map count"));
 
-			// First input is shorter and will be padded with null entries
+			// First input is shorter and will be padded with null entries (set maps)
 			{
 				int32 NumPredicateCalls = 0;
 				int32 NumNullEntries = 0;
@@ -832,7 +1271,110 @@ namespace UE::UAF::Tests
 				AddErrorIfFalse(NumNullEntries == 3, TEXT("Unexpected OuterJoin iteration result"));
 			}
 
-			// Second input is shorter and will be padded with null entries
+			// First input is shorter and will be padded with null entries (name/value maps)
+			{
+				int32 NumPredicateCalls = 0;
+				int32 NumNullEntries = 0;
+				OuterJoinBy(
+					FAttributeCollectionMapJoinOp(),
+					[this, &NumPredicateCalls, &NumNullEntries, &Collection0, &Collection2](FAttributeMap* Map2, FAttributeMap* Map0)
+					{
+						NumPredicateCalls++;
+						AddErrorIfFalse(CollectionContains(Collection0, Map0), TEXT("Unexpected OuterJoin iteration result"));
+
+						if (Map2 != nullptr)
+						{
+							AddErrorIfFalse(CollectionContains(Collection2, Map2), TEXT("Unexpected OuterJoin iteration result"));
+							AddErrorIfFalse(Map0 != nullptr && Map0->GetValueType() == Map2->GetValueType(), TEXT("Unexpected OuterJoin iteration result"));
+						}
+						else
+						{
+							NumNullEntries++;
+							AddErrorIfFalse(Map0 != nullptr, TEXT("Unexpected OuterJoin iteration result"));
+						}
+					},
+					Collection2.CreateMapIterator(),
+					Collection0.CreateMapIterator());
+				AddErrorIfFalse(NumPredicateCalls == 2, TEXT("Unexpected OuterJoin iteration result"));
+				AddErrorIfFalse(NumNullEntries == 1, TEXT("Unexpected OuterJoin iteration result"));
+
+				NumPredicateCalls = 0;
+				NumNullEntries = 0;
+				OuterJoinBy(
+					FAttributeCollectionMapJoinOp(),
+					[this, &NumPredicateCalls, &NumNullEntries, &Collection0, &Collection2](const FAttributeMap* Map2, const FAttributeMap* Map0)
+					{
+						NumPredicateCalls++;
+						AddErrorIfFalse(CollectionContains(Collection0, Map0), TEXT("Unexpected OuterJoin iteration result"));
+
+						if (Map2 != nullptr)
+						{
+							AddErrorIfFalse(CollectionContains(Collection2, Map2), TEXT("Unexpected OuterJoin iteration result"));
+							AddErrorIfFalse(Map0 != nullptr && Map0->GetValueType() == Map2->GetValueType(), TEXT("Unexpected OuterJoin iteration result"));
+						}
+						else
+						{
+							NumNullEntries++;
+							AddErrorIfFalse(Map0 != nullptr, TEXT("Unexpected OuterJoin iteration result"));
+						}
+					},
+					Collection2.CreateConstMapIterator(),
+					Collection0.CreateConstMapIterator());
+				AddErrorIfFalse(NumPredicateCalls == 2, TEXT("Unexpected OuterJoin iteration result"));
+				AddErrorIfFalse(NumNullEntries == 1, TEXT("Unexpected OuterJoin iteration result"));
+
+				NumPredicateCalls = 0;
+				NumNullEntries = 0;
+				OuterJoinBy(
+					FAttributeCollectionMapJoinOp(),
+					[this, &NumPredicateCalls, &NumNullEntries, &Collection0, &Collection2](const FAttributeMap* Map2, FAttributeMap* Map0)
+					{
+						NumPredicateCalls++;
+						AddErrorIfFalse(CollectionContains(Collection0, Map0), TEXT("Unexpected OuterJoin iteration result"));
+
+						if (Map2 != nullptr)
+						{
+							AddErrorIfFalse(CollectionContains(Collection2, Map2), TEXT("Unexpected OuterJoin iteration result"));
+							AddErrorIfFalse(Map0 != nullptr && Map0->GetValueType() == Map2->GetValueType(), TEXT("Unexpected OuterJoin iteration result"));
+						}
+						else
+						{
+							NumNullEntries++;
+							AddErrorIfFalse(Map0 != nullptr, TEXT("Unexpected OuterJoin iteration result"));
+						}
+					},
+					Collection2.CreateConstMapIterator(),
+					Collection0.CreateMapIterator());
+				AddErrorIfFalse(NumPredicateCalls == 2, TEXT("Unexpected OuterJoin iteration result"));
+				AddErrorIfFalse(NumNullEntries == 1, TEXT("Unexpected OuterJoin iteration result"));
+
+				NumPredicateCalls = 0;
+				NumNullEntries = 0;
+				OuterJoinBy(
+					FAttributeCollectionMapJoinOp(),
+					[this, &NumPredicateCalls, &NumNullEntries, &Collection0, &Collection2](FAttributeMap* Map2, const FAttributeMap* Map0)
+					{
+						NumPredicateCalls++;
+						AddErrorIfFalse(CollectionContains(Collection0, Map0), TEXT("Unexpected OuterJoin iteration result"));
+
+						if (Map2 != nullptr)
+						{
+							AddErrorIfFalse(CollectionContains(Collection2, Map2), TEXT("Unexpected OuterJoin iteration result"));
+							AddErrorIfFalse(Map0 != nullptr && Map0->GetValueType() == Map2->GetValueType(), TEXT("Unexpected OuterJoin iteration result"));
+						}
+						else
+						{
+							NumNullEntries++;
+							AddErrorIfFalse(Map0 != nullptr, TEXT("Unexpected OuterJoin iteration result"));
+						}
+					},
+					Collection2.CreateMapIterator(),
+					Collection0.CreateConstMapIterator());
+				AddErrorIfFalse(NumPredicateCalls == 2, TEXT("Unexpected OuterJoin iteration result"));
+				AddErrorIfFalse(NumNullEntries == 1, TEXT("Unexpected OuterJoin iteration result"));
+			}
+
+			// Second input is shorter and will be padded with null entries (set maps)
 			{
 				int32 NumPredicateCalls = 0;
 				int32 NumNullEntries = 0;
@@ -939,7 +1481,110 @@ namespace UE::UAF::Tests
 				AddErrorIfFalse(NumNullEntries == 3, TEXT("Unexpected OuterJoin iteration result"));
 			}
 
-			// Join with a transformer list
+			// Second input is shorter and will be padded with null entries (name/value maps)
+			{
+				int32 NumPredicateCalls = 0;
+				int32 NumNullEntries = 0;
+				OuterJoinBy(
+					FAttributeCollectionMapJoinOp(),
+					[this, &NumPredicateCalls, &NumNullEntries, &Collection0, &Collection2](FAttributeMap* Map0, FAttributeMap* Map2)
+					{
+						NumPredicateCalls++;
+						AddErrorIfFalse(CollectionContains(Collection0, Map0), TEXT("Unexpected OuterJoin iteration result"));
+
+						if (Map2 != nullptr)
+						{
+							AddErrorIfFalse(CollectionContains(Collection2, Map2), TEXT("Unexpected OuterJoin iteration result"));
+							AddErrorIfFalse(Map0 != nullptr && Map0->GetValueType() == Map2->GetValueType(), TEXT("Unexpected OuterJoin iteration result"));
+						}
+						else
+						{
+							NumNullEntries++;
+							AddErrorIfFalse(Map0 != nullptr, TEXT("Unexpected OuterJoin iteration result"));
+						}
+					},
+					Collection0.CreateMapIterator(),
+					Collection2.CreateMapIterator());
+				AddErrorIfFalse(NumPredicateCalls == 2, TEXT("Unexpected OuterJoin iteration result"));
+				AddErrorIfFalse(NumNullEntries == 1, TEXT("Unexpected OuterJoin iteration result"));
+
+				NumPredicateCalls = 0;
+				NumNullEntries = 0;
+				OuterJoinBy(
+					FAttributeCollectionMapJoinOp(),
+					[this, &NumPredicateCalls, &NumNullEntries, &Collection0, &Collection2](const FAttributeMap* Map0, const FAttributeMap* Map2)
+					{
+						NumPredicateCalls++;
+						AddErrorIfFalse(CollectionContains(Collection0, Map0), TEXT("Unexpected OuterJoin iteration result"));
+
+						if (Map2 != nullptr)
+						{
+							AddErrorIfFalse(CollectionContains(Collection2, Map2), TEXT("Unexpected OuterJoin iteration result"));
+							AddErrorIfFalse(Map0 != nullptr && Map0->GetValueType() == Map2->GetValueType(), TEXT("Unexpected OuterJoin iteration result"));
+						}
+						else
+						{
+							NumNullEntries++;
+							AddErrorIfFalse(Map0 != nullptr, TEXT("Unexpected OuterJoin iteration result"));
+						}
+					},
+					Collection0.CreateConstMapIterator(),
+					Collection2.CreateConstMapIterator());
+				AddErrorIfFalse(NumPredicateCalls == 2, TEXT("Unexpected OuterJoin iteration result"));
+				AddErrorIfFalse(NumNullEntries == 1, TEXT("Unexpected OuterJoin iteration result"));
+
+				NumPredicateCalls = 0;
+				NumNullEntries = 0;
+				OuterJoinBy(
+					FAttributeCollectionMapJoinOp(),
+					[this, &NumPredicateCalls, &NumNullEntries, &Collection0, &Collection2](const FAttributeMap* Map0, FAttributeMap* Map2)
+					{
+						NumPredicateCalls++;
+						AddErrorIfFalse(CollectionContains(Collection0, Map0), TEXT("Unexpected OuterJoin iteration result"));
+
+						if (Map2 != nullptr)
+						{
+							AddErrorIfFalse(CollectionContains(Collection2, Map2), TEXT("Unexpected OuterJoin iteration result"));
+							AddErrorIfFalse(Map0 != nullptr && Map0->GetValueType() == Map2->GetValueType(), TEXT("Unexpected OuterJoin iteration result"));
+						}
+						else
+						{
+							NumNullEntries++;
+							AddErrorIfFalse(Map0 != nullptr, TEXT("Unexpected OuterJoin iteration result"));
+						}
+					},
+					Collection0.CreateConstMapIterator(),
+					Collection2.CreateMapIterator());
+				AddErrorIfFalse(NumPredicateCalls == 2, TEXT("Unexpected OuterJoin iteration result"));
+				AddErrorIfFalse(NumNullEntries == 1, TEXT("Unexpected OuterJoin iteration result"));
+
+				NumPredicateCalls = 0;
+				NumNullEntries = 0;
+				OuterJoinBy(
+					FAttributeCollectionMapJoinOp(),
+					[this, &NumPredicateCalls, &NumNullEntries, &Collection0, &Collection2](FAttributeMap* Map0, const FAttributeMap* Map2)
+					{
+						NumPredicateCalls++;
+						AddErrorIfFalse(CollectionContains(Collection0, Map0), TEXT("Unexpected OuterJoin iteration result"));
+
+						if (Map2 != nullptr)
+						{
+							AddErrorIfFalse(CollectionContains(Collection2, Map2), TEXT("Unexpected OuterJoin iteration result"));
+							AddErrorIfFalse(Map0 != nullptr && Map0->GetValueType() == Map2->GetValueType(), TEXT("Unexpected OuterJoin iteration result"));
+						}
+						else
+						{
+							NumNullEntries++;
+							AddErrorIfFalse(Map0 != nullptr, TEXT("Unexpected OuterJoin iteration result"));
+						}
+					},
+					Collection0.CreateMapIterator(),
+					Collection2.CreateConstMapIterator());
+				AddErrorIfFalse(NumPredicateCalls == 2, TEXT("Unexpected OuterJoin iteration result"));
+				AddErrorIfFalse(NumNullEntries == 1, TEXT("Unexpected OuterJoin iteration result"));
+			}
+
+			// Join with a transformer list (set maps)
 			{
 				// Use a float transformer
 				// We set up Collection0 such that it has two mappings with float values: Transform -> Float and Float -> Float
@@ -951,7 +1596,7 @@ namespace UE::UAF::Tests
 					using TransformerType = Transformers::FSanitize;
 					using ValueType = FFloatAnimationAttribute;
 
-					static void Transform(FAttributeSetMap* SetMap) {}	// No-op
+					static void TransformSetMap(FAttributeSetMap* SetMap) {}	// No-op
 				};
 
 				// Dummy vector sanitize transformer
@@ -960,7 +1605,7 @@ namespace UE::UAF::Tests
 					using TransformerType = Transformers::FSanitize;
 					using ValueType = FVectorAnimationAttribute;
 
-					static void Transform(FAttributeSetMap* SetMap) {}	// No-op
+					static void TransformSetMap(FAttributeSetMap* SetMap) {}	// No-op
 				};
 
 #ifdef _MSC_VER
@@ -969,8 +1614,8 @@ namespace UE::UAF::Tests
 #endif
 
 				FAttributeValueTransformerList TransformerList(Transformers::FSanitize::TransformerName);
-				TransformerList.AddTransformer(FFloatAnimationAttribute::StaticStruct(), reinterpret_cast<FRawTransformerFunc>(&FFloatAttributeTransformer_SanitizeTest::Transform));
-				TransformerList.AddTransformer(FVectorAnimationAttribute::StaticStruct(), reinterpret_cast<FRawTransformerFunc>(&FVectorAttributeTransformer_SanitizeTest::Transform));
+				TransformerList.AddSetMapTransformer(FFloatAnimationAttribute::StaticStruct(), reinterpret_cast<FRawTransformerFunc>(&FFloatAttributeTransformer_SanitizeTest::TransformSetMap));
+				TransformerList.AddSetMapTransformer(FVectorAnimationAttribute::StaticStruct(), reinterpret_cast<FRawTransformerFunc>(&FVectorAttributeTransformer_SanitizeTest::TransformSetMap));
 
 #ifdef _MSC_VER
 	#pragma warning(pop)
@@ -980,7 +1625,7 @@ namespace UE::UAF::Tests
 				int32 NumNullCollectionEntries = 0;
 				int32 NumNullTransformEntries = 0;
 				OuterJoinBy(
-					FAttributeValueTransformerListWithAttributeCollectionJoinOp(),
+					FAttributeValueTransformerListWithAttributeCollectionSetMapJoinOp(),
 					[this, &NumPredicateCalls, &NumNullCollectionEntries, &NumNullTransformEntries, &Collection0](FRawTransformerFunc Transformer, FAttributeSetMap* SetMap)
 					{
 						NumPredicateCalls++;
@@ -1004,7 +1649,7 @@ namespace UE::UAF::Tests
 							NumNullTransformEntries++;
 						}
 					},
-					TransformerList.CreateTransformerIterator(),
+					TransformerList.CreateSetMapTransformerIterator(),
 					Collection0.CreateSetMapIterator());
 
 				// [Transformer Value Type, Collection Set Map Value Type]
@@ -1021,6 +1666,85 @@ namespace UE::UAF::Tests
 				AddErrorIfFalse(NumNullCollectionEntries == 2, TEXT("Unexpected OuterJoin iteration result"));
 				AddErrorIfFalse(NumNullTransformEntries == 3, TEXT("Unexpected OuterJoin iteration result"));
 			}
+
+			// Join with a transformer list (name/value maps)
+			{
+				// Use a float transformer
+				// We set up Collection0 such that it has two mappings with these value types: Float and Bone Transform
+				// InnerJoin should execute the predicate twice, once for each value type that matches the transformer's value type
+
+				// Dummy float sanitize transformer
+				struct FFloatAttributeTransformer_SanitizeTest
+				{
+					using TransformerType = Transformers::FSanitize;
+					using ValueType = FFloatAnimationAttribute;
+
+					static void TransformMap(FAttributeMap* Map) {}	// No-op
+				};
+
+				// Dummy vector sanitize transformer
+				struct FVectorAttributeTransformer_SanitizeTest
+				{
+					using TransformerType = Transformers::FSanitize;
+					using ValueType = FVectorAnimationAttribute;
+
+					static void TransformMap(FAttributeMap* Map) {}	// No-op
+				};
+
+#ifdef _MSC_VER
+	#pragma warning(push)
+	#pragma warning(disable : 4191) // unsafe conversion from 'function pointer type A' to 'function pointer type B'
+#endif
+
+				FAttributeValueTransformerList TransformerList(Transformers::FSanitize::TransformerName);
+				TransformerList.AddMapTransformer(FFloatAnimationAttribute::StaticStruct(), reinterpret_cast<FRawTransformerFunc>(&FFloatAttributeTransformer_SanitizeTest::TransformMap));
+				TransformerList.AddMapTransformer(FVectorAnimationAttribute::StaticStruct(), reinterpret_cast<FRawTransformerFunc>(&FVectorAttributeTransformer_SanitizeTest::TransformMap));
+
+#ifdef _MSC_VER
+	#pragma warning(pop)
+#endif
+
+				int32 NumPredicateCalls = 0;
+				int32 NumNullCollectionEntries = 0;
+				int32 NumNullTransformEntries = 0;
+				OuterJoinBy(
+					FAttributeValueTransformerListWithAttributeCollectionMapJoinOp(),
+					[this, &NumPredicateCalls, &NumNullCollectionEntries, &NumNullTransformEntries, &Collection0](FRawTransformerFunc Transformer, FAttributeMap* Map)
+					{
+						NumPredicateCalls++;
+
+						if (Map != nullptr)
+						{
+							AddErrorIfFalse(CollectionContains(Collection0, Map), TEXT("Unexpected OuterJoin iteration result"));
+
+							if (Transformer != nullptr)
+							{
+								AddErrorIfFalse(Map->GetValueType() == FFloatAnimationAttribute::StaticStruct(), TEXT("Unexpected OuterJoin iteration result"));
+							}
+						}
+						else
+						{
+							NumNullCollectionEntries++;
+						}
+
+						if (Transformer == nullptr)
+						{
+							NumNullTransformEntries++;
+						}
+					},
+					TransformerList.CreateMapTransformerIterator(),
+					Collection0.CreateMapIterator());
+
+				// [Transformer Value Type, Collection Map Value Type]
+				// ---
+				// [Float, Float]
+				// [nullptr, Bone]
+				// [Vector, nullptr]
+
+				AddErrorIfFalse(NumPredicateCalls == 3, TEXT("Unexpected OuterJoin iteration result"));
+				AddErrorIfFalse(NumNullCollectionEntries == 1, TEXT("Unexpected OuterJoin iteration result"));
+				AddErrorIfFalse(NumNullTransformEntries == 1, TEXT("Unexpected OuterJoin iteration result"));
+			}
 		}
 
 		return true;
