Clean up automated UAF Workspace test cases

diff --git a/Engine/Plugins/Experimental/UAF/UAFTests/Source/UAFTests/Private/UAFAssetCreationTests.cpp b/Engine/Plugins/Experimental/UAF/UAFTests/Source/UAFTests/Private/UAFAssetCreationTests.cpp
--- a/Engine/Plugins/Experimental/UAF/UAFTests/Source/UAFTests/Private/UAFAssetCreationTests.cpp
+++ b/Engine/Plugins/Experimental/UAF/UAFTests/Source/UAFTests/Private/UAFAssetCreationTests.cpp
@@ -7,8 +7,6 @@
 #include "CQTest.h"
 #include "Editor.h"
 #include "Subsystems/AssetEditorSubsystem.h"
-#include "WorkspaceEditor/Private/Workspace.h"
-#include "Workspace/AnimNextWorkspaceFactory.h"
 #include "AnimNextStateTree.h"
 #include "AnimNextStateTreeFactory.h"
 #include "Graph/AnimNextAnimationGraph.h"
diff --git a/Engine/Plugins/Experimental/UAF/UAFTests/Source/UAFTests/Private/UAFModuleTests.cpp b/Engine/Plugins/Experimental/UAF/UAFTests/Source/UAFTests/Private/UAFModuleTests.cpp
--- a/Engine/Plugins/Experimental/UAF/UAFTests/Source/UAFTests/Private/UAFModuleTests.cpp
+++ b/Engine/Plugins/Experimental/UAF/UAFTests/Source/UAFTests/Private/UAFModuleTests.cpp
@@ -8,8 +8,6 @@
 #include "Editor.h"
 #include "ScopedTransaction.h"
 #include "Subsystems/AssetEditorSubsystem.h"
-#include "WorkspaceEditor/Private/Workspace.h"
-#include "Workspace/AnimNextWorkspaceFactory.h"
 #include "Module/AnimNextModule.h"
 #include "Module/AnimNextModuleFactory.h"
 #include "AnimNextRigVMAssetEditorData.h"
@@ -21,11 +19,8 @@
 
 TEST_CLASS(FAnimNextModuleTests, "Animation.UAF.Functional")
 {
-	UWorkspace* AnimNextWorkspace = nullptr;
 	UAnimNextModule* AnimNextModule = nullptr;
-	UPackage* AnimNextWorkspacePackage = nullptr;
 	UPackage* AnimNextModulePackage = nullptr;
-	FName AnimNextWorkspacePackageName;
 	FName AnimNextModulePackageName;
 	UAssetEditorSubsystem* AssetEditorSubsystem = nullptr;
 	UAnimNextRigVMAssetEditorData* EditorData = nullptr;
@@ -50,12 +45,6 @@ TEST_CLASS(FAnimNextModuleTests, "Animation.UAF.Functional")
 	
 	AFTER_EACH()
 	{
-		if (AnimNextWorkspace)
-		{
-			ObjectTools::DeleteSingleObject(AnimNextWorkspace);
-			AnimNextWorkspace = nullptr;
-		}
-		
 		if (AnimNextModule)
 		{
 			ObjectTools::DeleteSingleObject(AnimNextModule);
diff --git a/Engine/Plugins/Experimental/UAF/UAFTests/Source/UAFTests/Private/UAFWorkspaceTests.cpp b/Engine/Plugins/Experimental/UAF/UAFTests/Source/UAFTests/Private/UAFWorkspaceTests.cpp
--- a/Engine/Plugins/Experimental/UAF/UAFTests/Source/UAFTests/Private/UAFWorkspaceTests.cpp
+++ b/Engine/Plugins/Experimental/UAF/UAFTests/Source/UAFTests/Private/UAFWorkspaceTests.cpp
@@ -9,7 +9,6 @@
 #include "IWorkspaceEditor.h"
 #include "Subsystems/AssetEditorSubsystem.h"
 #include "WorkspaceEditor/Private/Workspace.h"
-#include "Workspace/AnimNextWorkspaceFactory.h"
 #include "Graph/AnimNextAnimationGraph.h"
 #include "Graph/AnimNextAnimationGraphFactory.h"
 #include "Module/AnimNextModule.h"
@@ -45,6 +44,7 @@ TEST_CLASS(FUAFWorkspaceTests, "Animation.UAF.Functional")
 			});
 	}
 
+	// QMetry Test Case:  UE-TC-22353
 	TEST_METHOD(Verify_Cannot_Remove_Workspace_Root_Asset)
 	{
 		TestCommandBuilder
@@ -70,6 +70,7 @@ TEST_CLASS(FUAFWorkspaceTests, "Animation.UAF.Functional")
 			});
 	}
 
+	// QMetry Test Case:  UE-TC-10716
 	TEST_METHOD(Verify_Cannot_Add_Multiple_Assets_To_Workspace)
 	{
 		TestCommandBuilder
