Update AddUnitNode API to support Unit Node Class Path
- Added additional optional parameterto pass in the class path
- Replaced hard-coded script strings, instead using GetPathName() API

diff --git a/Engine/Plugins/Experimental/UAF/UAFTests/Source/UAFTests/Private/UAFAnimationGraphTests.cpp b/Engine/Plugins/Experimental/UAF/UAFTests/Source/UAFTests/Private/UAFAnimationGraphTests.cpp
--- a/Engine/Plugins/Experimental/UAF/UAFTests/Source/UAFTests/Private/UAFAnimationGraphTests.cpp
+++ b/Engine/Plugins/Experimental/UAF/UAFTests/Source/UAFTests/Private/UAFAnimationGraphTests.cpp
@@ -10,6 +10,8 @@
 #include "Subsystems/AssetEditorSubsystem.h"
 #include "Graph/AnimNextAnimationGraph.h"
 #include "Graph/AnimNextAnimationGraphFactory.h"
+#include "Graph/RigUnit_AnimNextTraitStack.h"
+#include "RigVMFunctions/Simulation/RigVMFunction_Accumulate.h"
 #include "AnimNextRigVMAssetEditorData.h"
 #include "Math/Vector2D.h"
 
@@ -82,7 +84,7 @@ TEST_CLASS(FAnimNextAnimationGraphTests, "Animation.UAF.Functional")
 			.Then(TEXT("Create an AccumulateLerp Transform node"), [&]()
 			{
 				TArray<UEdGraphPin*> FromPins;
-				UEdGraphNode* AccumulateLerpTransformNode = UAFTestsUtilities::AddUnitNode(ParentGraph, "/Script/RigVM.RigVMFunction_AccumulateTransformLerp", FromPins, FVector2f(-525.0f, 125.0f));
+				UEdGraphNode* AccumulateLerpTransformNode = UAFTestsUtilities::AddUnitNode(ParentGraph, FRigVMFunction_AccumulateTransformLerp::StaticStruct()->GetPathName(), FromPins, FVector2f(-525.0f, 125.0f));
 				ASSERT_THAT(IsNotNull(AccumulateLerpTransformNode));
 			})
 			.Then(TEXT("Create a bool Variable"), [&]()
@@ -184,7 +186,7 @@ protected:
 		TestRunner->AddExpectedError("Base Trait Data is Invalid. Please, select a new Base Trait.", EAutomationExpectedErrorFlags::Contains, -1);
 		
 		TArray<UEdGraphPin*> FromPins;
-		UEdGraphNode* TraitStackNode = UAFTestsUtilities::AddUnitNode(ParentGraph, "/Script/AnimNextAnimGraph.RigUnit_AnimNextTraitStack", FromPins, FVector2f(0.0f, 0.0f));
+		UEdGraphNode* TraitStackNode = UAFTestsUtilities::AddUnitNode(ParentGraph, FRigUnit_AnimNextTraitStack::StaticStruct()->GetPathName(), FromPins, FVector2f(0.0f, 0.0f), "/Script/UAFAnimGraphUncookedOnly.AnimNextTraitStackUnitNode");
 		ASSERT_THAT(IsNotNull(TraitStackNode));
 		ASSERT_THAT(IsNotNull(ControllerGraph->FindNodeByName(FName(TraitStackNodeName))));	
 	}
diff --git a/Engine/Plugins/Experimental/UAF/UAFTests/Source/UAFTests/Private/UAFModuleTests.cpp b/Engine/Plugins/Experimental/UAF/UAFTests/Source/UAFTests/Private/UAFModuleTests.cpp
--- a/Engine/Plugins/Experimental/UAF/UAFTests/Source/UAFTests/Private/UAFModuleTests.cpp
+++ b/Engine/Plugins/Experimental/UAF/UAFTests/Source/UAFTests/Private/UAFModuleTests.cpp
@@ -10,6 +10,8 @@
 #include "Subsystems/AssetEditorSubsystem.h"
 #include "Module/AnimNextModule.h"
 #include "Module/AnimNextModuleFactory.h"
+#include "Graph/RigUnit_AnimNextRunAnimationGraph_v2.h"
+#include "Graph/RigUnit_AnimNextWriteSkeletalMeshComponentPose.h"
 #include "AnimNextRigVMAssetEditorData.h"
 #include "Math/Vector2D.h"
 #include "AssetToolsModule.h"
@@ -70,7 +72,7 @@ TEST_CLASS(FAnimNextModuleTests, "Animation.UAF.Functional")
 				ASSERT_THAT(IsNotNull(ParentGraph));
 				
 				TArray<UEdGraphPin*> FromPins;
-				WritePoseNode = UAFTestsUtilities::AddUnitNode(ParentGraph, "/Script/AnimNext.RigUnit_AnimNextWriteSkeletalMeshComponentPose", FromPins, FVector2f(250.0f, 135.0f));
+				WritePoseNode = UAFTestsUtilities::AddUnitNode(ParentGraph, FRigUnit_AnimNextWriteSkeletalMeshComponentPose::StaticStruct()->GetPathName(), FromPins, FVector2f(250.0f, 135.0f));
 				ASSERT_THAT(IsNotNull(WritePoseNode));
 			})
 			.Then(TEXT("Create a bool Variable"), [&]()
@@ -157,13 +159,13 @@ TEST_CLASS(FAnimNextModuleTests, "Animation.UAF.Functional")
 				ASSERT_THAT(IsNotNull(ParentGraph));	
 				
 				TArray<UEdGraphPin*> FromPins;
-				RunGraphNode = UAFTestsUtilities::AddUnitNode(ParentGraph, "/Script/AnimNextAnimGraph.RigUnit_AnimNextRunAnimationGraph_v2", FromPins, FVector2f(64.0f, 16.0f));
+				RunGraphNode = UAFTestsUtilities::AddUnitNode(ParentGraph, FRigUnit_AnimNextRunAnimationGraph_v2::StaticStruct()->GetPathName(), FromPins, FVector2f(64.0f, 16.0f));
 				ASSERT_THAT(IsNotNull(RunGraphNode));
 			})
 			.Then(TEXT("Create a Write Pose node"), [&]()
 			{
 				TArray<UEdGraphPin*> FromPins;
-				WritePoseNode = UAFTestsUtilities::AddUnitNode(ParentGraph, "/Script/AnimNext.RigUnit_AnimNextWriteSkeletalMeshComponentPose", FromPins, FVector2f(336.0f, 0.0f));
+				WritePoseNode = UAFTestsUtilities::AddUnitNode(ParentGraph, FRigUnit_AnimNextWriteSkeletalMeshComponentPose::StaticStruct()->GetPathName(), FromPins, FVector2f(336.0f, 0.0f));
 				ASSERT_THAT(IsNotNull(WritePoseNode));
 			})
 			.Then(TEXT("Add link from PrePhysics node Execute pin to Run Graph node Execute pin"), [&]()
@@ -222,13 +224,13 @@ TEST_CLASS(FAnimNextModuleTests, "Animation.UAF.Functional")
 				ASSERT_THAT(IsNotNull(ParentGraph));	
 				
 				TArray<UEdGraphPin*> FromPins;
-				RunGraphNode = UAFTestsUtilities::AddUnitNode(ParentGraph, "/Script/AnimNextAnimGraph.RigUnit_AnimNextRunAnimationGraph_v2", FromPins, FVector2f(64.0f, 16.0f));
+				RunGraphNode = UAFTestsUtilities::AddUnitNode(ParentGraph, FRigUnit_AnimNextRunAnimationGraph_v2::StaticStruct()->GetPathName(), FromPins, FVector2f(64.0f, 16.0f));
 				ASSERT_THAT(IsNotNull(RunGraphNode));
 			})
 			.Then(TEXT("Create a Write Pose node"), [&]()
 			{
 				TArray<UEdGraphPin*> FromPins;
-				WritePoseNode = UAFTestsUtilities::AddUnitNode(ParentGraph, "/Script/AnimNext.RigUnit_AnimNextWriteSkeletalMeshComponentPose", FromPins, FVector2f(336.0f, 0.0f));
+				WritePoseNode = UAFTestsUtilities::AddUnitNode(ParentGraph, FRigUnit_AnimNextWriteSkeletalMeshComponentPose::StaticStruct()->GetPathName(), FromPins, FVector2f(336.0f, 0.0f));
 				ASSERT_THAT(IsNotNull(WritePoseNode));
 			})
 			.Then(TEXT("Add link from PrePhysics node Execute pin to Run Graph node Execute pin"), [&]()
diff --git a/Engine/Plugins/Experimental/UAF/UAFTests/Source/UAFTests/Private/UAFTestsUtilities.cpp b/Engine/Plugins/Experimental/UAF/UAFTests/Source/UAFTests/Private/UAFTestsUtilities.cpp
--- a/Engine/Plugins/Experimental/UAF/UAFTests/Source/UAFTests/Private/UAFTestsUtilities.cpp
+++ b/Engine/Plugins/Experimental/UAF/UAFTests/Source/UAFTests/Private/UAFTestsUtilities.cpp
@@ -27,7 +27,7 @@ namespace UAFTestsUtilities
 		return InFactory->FactoryCreateNew(InClass, NewPackage, NewPackageName, RF_Public | RF_Standalone, NULL, GWarn);
 	}
 
-	UEdGraphNode* AddUnitNode(UEdGraph* InParentGraph, const FString& InScriptStructPath, TArray<UEdGraphPin*>& InFromPins, const FVector2f& InRigUnitLocation)
+	UEdGraphNode* AddUnitNode(UEdGraph* InParentGraph, const FString& InScriptStructPath, TArray<UEdGraphPin*>& InFromPins, const FVector2f& InRigUnitLocation, const FString& InUnitNodeClassPath)
 	{		
 		if (!InParentGraph)
 		{
@@ -39,7 +39,22 @@ namespace UAFTestsUtilities
 		FText Category = NSLOCTEXT("Category", "Category", "Category");
 		FText MenuDesc = NSLOCTEXT("Menu_Desc", "MenuDesc", "Menu Desc");
 		FText ToolTip = NSLOCTEXT("Tool_Tip", "ToolTip", "Tool Tip");
-		TSharedPtr<FAnimNextSchemaAction_RigUnit> RigUnitAction = MakeShared<FAnimNextSchemaAction_RigUnit>(URigVMUnitNode::StaticClass(), ScriptStruct, Category, MenuDesc, ToolTip);
+		
+		UClass* UnitNodeClass = nullptr;
+		if (InUnitNodeClassPath.IsEmpty())
+		{
+			UnitNodeClass = URigVMUnitNode::StaticClass();
+		}
+		else
+		{
+			UnitNodeClass = RigVMTypeUtils::FindObjectFromCPPTypeObjectPath<UClass>(InUnitNodeClassPath);
+		}
+		
+		TSharedPtr<FAnimNextSchemaAction_RigUnit> RigUnitAction = MakeShared<FAnimNextSchemaAction_RigUnit>(UnitNodeClass,
+																											ScriptStruct,
+																											Category,
+																											MenuDesc,
+																											ToolTip);
 
 		return RigUnitAction->PerformAction(InParentGraph, InFromPins, InRigUnitLocation);
 	}
diff --git a/Engine/Plugins/Experimental/UAF/UAFTests/Source/UAFTests/Public/UAFTestsUtilities.h b/Engine/Plugins/Experimental/UAF/UAFTests/Source/UAFTests/Public/UAFTestsUtilities.h
--- a/Engine/Plugins/Experimental/UAF/UAFTests/Source/UAFTests/Public/UAFTestsUtilities.h
+++ b/Engine/Plugins/Experimental/UAF/UAFTests/Source/UAFTests/Public/UAFTestsUtilities.h
@@ -32,9 +32,10 @@ namespace UAFTestsUtilities
 	 * @param InScriptStructPath:  Path to the script struct.
 	 * @param InFromPins:  Pin to wire from.
 	 * @param InRigUnitLocation:  Location of the RigUnit node on the graph.
+	 * @param InUnitNodeClassPath:  Path to the unit node class path, optional
 	 * @return A pointer to the RigUnit node instance, or nullptr otherwise.
 	 */
-	UEdGraphNode* AddUnitNode(UEdGraph* InParentGraph, const FString& InScriptStructPath, TArray<UEdGraphPin*>& InFromPins, const FVector2f& InRigUnitLocation);
+	UEdGraphNode* AddUnitNode(UEdGraph* InParentGraph, const FString& InScriptStructPath, TArray<UEdGraphPin*>& InFromPins, const FVector2f& InRigUnitLocation, const FString& InUnitNodeClassPath = TEXT(""));
 
 	/**
 	 * Adds a new function node instance to the AnimNext Asset specified.
