Changing variable types on non-bound object variables can lead to checks and potential crashes
#fix guard against invalid/unsupported property types when compiling variables from data/native
#misc removed/relaxed requirement for asset to be compiled to have a valid message log, compiling variables (for combined property-bag) does not queue the referenced asset for compilation (not needed)
#misc removed invalid (read unsupported) variable types from automated tests
#Synced-CL 48715962

diff --git a/Engine/Plugins/Experimental/UAF/UAF/Source/UAFUncookedOnly/Private/AnimNextRigVMAssetEditorData.cpp b/Engine/Plugins/Experimental/UAF/UAF/Source/UAFUncookedOnly/Private/AnimNextRigVMAssetEditorData.cpp
--- a/Engine/Plugins/Experimental/UAF/UAF/Source/UAFUncookedOnly/Private/AnimNextRigVMAssetEditorData.cpp
+++ b/Engine/Plugins/Experimental/UAF/UAF/Source/UAFUncookedOnly/Private/AnimNextRigVMAssetEditorData.cpp
@@ -3255,6 +3255,14 @@ FInstancedPropertyBag UAnimNextRigVMAssetEditorData::GenerateCombinedPropertyBag
 		TSharedRef<UE::UAF::FStructDataCache> StructData = UE::UAF::FStructDataCache::GetStructInfo(ReferencedStruct);
 		for (const UE::UAF::FStructDataCache::FPropertyInfo& PropertyInfo : StructData->GetProperties())
 		{
+			const FAnimNextParamType ParamType =  FAnimNextParamType::FromProperty(PropertyInfo.Property);
+			const FRigVMTemplateArgumentType RigVMArgumentType = ParamType.ToRigVMTemplateArgument();
+			if (!RigVMArgumentType.IsValid() || RigVMArgumentType.IsUnknownType())
+			{
+				InCompileContext.GetAssetCompileContext().Error(Asset, LOCTEXT("InvalidNativeVariableType", "@@ Variable '{0}' from '{1}' has unsupported variable type '{2}'"), FText::FromName(PropertyInfo.Property->GetFName()),  FText::FromName(ReferencedStruct->GetFName()), FText::FromString(ParamType.ToString()));
+				continue;
+			}
+			
 			constexpr bool bOverrideProperty = false;
 			const EPropertyBagAlterationResult Result = CombinedPropertyBag.AddProperty(PropertyInfo.Property->GetFName(), PropertyInfo.Property, bOverrideProperty);
 
diff --git a/Engine/Plugins/Experimental/UAF/UAF/Source/UAFUncookedOnly/Private/UAFCompilationScope.cpp b/Engine/Plugins/Experimental/UAF/UAF/Source/UAFUncookedOnly/Private/UAFCompilationScope.cpp
--- a/Engine/Plugins/Experimental/UAF/UAF/Source/UAFUncookedOnly/Private/UAFCompilationScope.cpp
+++ b/Engine/Plugins/Experimental/UAF/UAF/Source/UAFUncookedOnly/Private/UAFCompilationScope.cpp
@@ -162,7 +162,7 @@ FCompilationScope::~FCompilationScope()
 			
 				TSharedRef<FTokenizedMessage> Message = FTokenizedMessage::Create(EMessageSeverity::Info);
 				Message->AddToken(FUObjectToken::Create(LogPair.Key.Get()));
-				Message->AddText(FText::Format(LOCTEXT("AssetCompileTimeFormat", "in {0} ms."), FText::AsNumber(PerAssetTimeCompileTimeSpan[AssetIndex] * 1000.f, &TimeFormat)));
+				Message->AddText(FText::Format(LOCTEXT("AssetCompileTimeFormat", "in {0} ms."), FText::AsNumber(AssetIndex != INDEX_NONE ? PerAssetTimeCompileTimeSpan[AssetIndex] * 1000.f : 0.f, &TimeFormat)));
 				Message->SetIndentationLevel(AssetMessageIndentationLevel);
 				MessageLog.AddMessage(Message);
 				
@@ -232,7 +232,6 @@ FCompilerResultsLog& FCompilationScope::GetLogForObject(const UObject* InObject)
 	check(InObject);
 
 	FCompileJobThreadData* ThreadData = &GCompileJobData;
-	check(ThreadData->CompiledAssets.Contains(InObject));
 	return ThreadData->PerObjectCompileLog.FindOrAdd(InObject);
 }
 
diff --git a/Engine/Plugins/Experimental/UAF/UAF/Source/UAFUncookedOnly/Private/UncookedOnlyUtils.cpp b/Engine/Plugins/Experimental/UAF/UAF/Source/UAFUncookedOnly/Private/UncookedOnlyUtils.cpp
--- a/Engine/Plugins/Experimental/UAF/UAF/Source/UAFUncookedOnly/Private/UncookedOnlyUtils.cpp
+++ b/Engine/Plugins/Experimental/UAF/UAF/Source/UAFUncookedOnly/Private/UncookedOnlyUtils.cpp
@@ -88,16 +88,17 @@ FInstancedPropertyBag FUtils::MakePropertyBagForEditorData(const UAnimNextRigVMA
 	UAnimNextRigVMAsset* Asset = GetAsset<UAnimNextRigVMAsset>(InEditorData);
 	auto AddVariable = [Asset, &NumPublicVariables, &StructEntryInfos, &EntryInfoIndexMap, &InCompileContext](FName InName, const FAnimNextParamType& InType, EAnimNextExportAccessSpecifier InAccess, TConstArrayView<uint8> InValue, EPropertyFlags InFlags, FGuid InStableGUID)
 	{
-		if(!InType.IsValid())
+		const FRigVMTemplateArgumentType RigVmArgumentType = InType.ToRigVMTemplateArgument();
+		if (!RigVmArgumentType.IsValid() || RigVmArgumentType.IsUnknownType())
 		{
-			InCompileContext.GetAssetCompileContext().Error(Asset, LOCTEXT("InvalidVariableTypeFound", " @@ Variable '{0}' with invalid type found"), FText::FromName(InName));
+			InCompileContext.GetAssetCompileContext().Error(Asset, LOCTEXT("InvalidVariableTypeFound", "@@ Variable '{0}' has unsupported variable type '{1}'"), FText::FromName(InName), FText::FromString(InType.ToString()));
 			return;
 		}
 
 		// Check for conflicts
 		if(EntryInfoIndexMap.Contains(InName))
 		{
-			InCompileContext.GetAssetCompileContext().Error(Asset, LOCTEXT("DuplicateVariableFound", " @@ Variable '{0}' with duplicate name found"), FText::FromName(InName));
+			InCompileContext.GetAssetCompileContext().Error(Asset, LOCTEXT("DuplicateVariableFound", "@@ Variable '{0}' with duplicate name found"), FText::FromName(InName));
 			return;
 		}
 
diff --git a/Engine/Plugins/Experimental/UAF/UAFTests/Source/UAFAnimGraphTestSuite/Private/AnimNextAnimGraphEditorVariablesTest.cpp b/Engine/Plugins/Experimental/UAF/UAFTests/Source/UAFAnimGraphTestSuite/Private/AnimNextAnimGraphEditorVariablesTest.cpp
--- a/Engine/Plugins/Experimental/UAF/UAFTests/Source/UAFAnimGraphTestSuite/Private/AnimNextAnimGraphEditorVariablesTest.cpp
+++ b/Engine/Plugins/Experimental/UAF/UAFTests/Source/UAFAnimGraphTestSuite/Private/AnimNextAnimGraphEditorVariablesTest.cpp
@@ -100,7 +100,6 @@ bool FEditor_AnimGraph_Variables::RunTest(const FString& InParameters)
 		TestVariableType(FAnimNextParamType::GetType<double>());
 		TestVariableType(FAnimNextParamType::GetType<FName>());
 		TestVariableType(FAnimNextParamType::GetType<FString>());
-		TestVariableType(FAnimNextParamType::GetType<FText>());
 		TestVariableType(FAnimNextParamType::GetType<EPropertyBagPropertyType>());
 		TestVariableType(FAnimNextParamType::GetType<FVector>());
 		TestVariableType(FAnimNextParamType::GetType<FQuat>());
@@ -250,4 +249,4 @@ bool FEditor_AnimGraph_Variables_Python::RunTest(const FString& InParameters)
 
 }
 
-#endif	// WITH_DEV_AUTOMATION_TESTS && WITH_EDITOR
\ No newline at end of file
+#endif	// WITH_DEV_AUTOMATION_TESTS && WITH_EDITOR
diff --git a/Engine/Plugins/Experimental/UAF/UAFTests/Source/UAFAnimGraphTestSuite/Private/AnimNextAnimGraphTraitGraphTest.cpp b/Engine/Plugins/Experimental/UAF/UAFTests/Source/UAFAnimGraphTestSuite/Private/AnimNextAnimGraphTraitGraphTest.cpp
--- a/Engine/Plugins/Experimental/UAF/UAFTests/Source/UAFAnimGraphTestSuite/Private/AnimNextAnimGraphTraitGraphTest.cpp
+++ b/Engine/Plugins/Experimental/UAF/UAFTests/Source/UAFAnimGraphTestSuite/Private/AnimNextAnimGraphTraitGraphTest.cpp
@@ -692,14 +692,11 @@ bool FAnimationAnimNextRuntimeTest_Variables::RunTest(const FString& InParameter
 		AddPublicVariable(AnimationGraph, "Bool", true);
 		AddPublicVariable(AnimationGraph, "Byte", (uint8)42);
 		AddPublicVariable(AnimationGraph, "Int32", (int32)-4679222);
-		AddPublicVariable(AnimationGraph, "UInt32", (uint32)3415919103);
 		AddPublicVariable(AnimationGraph, "Int64", (int64)-3415919105);
-		AddPublicVariable(AnimationGraph, "UInt64", (uint64)34159191067);
 		AddPublicVariable(AnimationGraph, "Float", 1.0f);
 		AddPublicVariable(AnimationGraph, "Double", 1.0);
 		AddPublicVariable(AnimationGraph, "Name", FName("Test"));
 		AddPublicVariable(AnimationGraph, "String", FString("Test"));
-		AddPublicVariable(AnimationGraph, "Text", NSLOCTEXT("Tests", "Test", "Test"));
 		AddPublicVariable(AnimationGraph, "Enum", EPropertyBagPropertyType::Double);
 		AddPublicVariable(AnimationGraph, "Struct", FVector::OneVector);
 		AddPublicVariable(AnimationGraph, "DerivedStruct", FTestDerivedVector());
@@ -709,6 +706,7 @@ bool FAnimationAnimNextRuntimeTest_Variables::RunTest(const FString& InParameter
 		AddPublicVariable(AnimationGraph, "SoftClass", FSoftClassPath(UAnimNextSharedVariables::StaticClass()));
 
 		TSharedPtr<FAnimNextGraphInstance> GraphInstance = AnimationGraph->AllocateInstance();
+		UE_RETURN_ON_ERROR(GraphInstance.IsValid(), "Failed to allocate animation graph instance");
 
 		// Bool/Integers + conversions
 		{
@@ -725,18 +723,10 @@ bool FAnimationAnimNextRuntimeTest_Variables::RunTest(const FString& InParameter
 			UE_RETURN_ON_ERROR(GraphInstance->GetVariable(FAnimNextVariableReference::FromName("Int32", AnimationGraph), TestInt32) == EPropertyBagResult::Success, "FAnimationAnimNextRuntimeTest_Variables -> GetVariable failed");
 			UE_RETURN_ON_ERROR(TestInt32 == -4679222, "FAnimationAnimNextRuntimeTest_Variables -> Variable value did not match");
 
-			uint32 TestUInt32 = 0;
-			UE_RETURN_ON_ERROR(GraphInstance->GetVariable(FAnimNextVariableReference::FromName("UInt32", AnimationGraph), TestUInt32) == EPropertyBagResult::Success, "FAnimationAnimNextRuntimeTest_Variables -> GetVariable failed");
-			UE_RETURN_ON_ERROR(TestUInt32 == 3415919103, "FAnimationAnimNextRuntimeTest_Variables -> Variable value did not match");
-
 			int64 TestInt64 = 0;
 			UE_RETURN_ON_ERROR(GraphInstance->GetVariable(FAnimNextVariableReference::FromName("Int64", AnimationGraph), TestInt64) == EPropertyBagResult::Success, "FAnimationAnimNextRuntimeTest_Variables -> GetVariable failed");
 			UE_RETURN_ON_ERROR(TestInt64 == -3415919105, "FAnimationAnimNextRuntimeTest_Variables -> Variable value did not match");
 
-			uint64 TestUInt64 = 0;
-			UE_RETURN_ON_ERROR(GraphInstance->GetVariable(FAnimNextVariableReference::FromName("UInt64", AnimationGraph), TestUInt64) == EPropertyBagResult::Success, "FAnimationAnimNextRuntimeTest_Variables -> GetVariable failed");
-			UE_RETURN_ON_ERROR(TestUInt64 == 34159191067, "FAnimationAnimNextRuntimeTest_Variables -> Variable value did not match");
-
 			// Conversions
 
 			// Bool
@@ -748,18 +738,10 @@ bool FAnimationAnimNextRuntimeTest_Variables::RunTest(const FString& InParameter
 			UE_RETURN_ON_ERROR(GraphInstance->GetVariable(FAnimNextVariableReference::FromName("Bool", AnimationGraph), TestInt32) == EPropertyBagResult::Success, "FAnimationAnimNextRuntimeTest_Variables -> GetVariable failed");
 			UE_RETURN_ON_ERROR(TestInt32 == 1, "FAnimationAnimNextRuntimeTest_Variables -> Variable value did not match");
 
-			TestUInt32 = 0;
-			UE_RETURN_ON_ERROR(GraphInstance->GetVariable(FAnimNextVariableReference::FromName("Bool", AnimationGraph), TestUInt32) == EPropertyBagResult::Success, "FAnimationAnimNextRuntimeTest_Variables -> GetVariable failed");
-			UE_RETURN_ON_ERROR(TestUInt32 == 1, "FAnimationAnimNextRuntimeTest_Variables -> Variable value did not match");
-
 			TestInt64 = 0;
 			UE_RETURN_ON_ERROR(GraphInstance->GetVariable(FAnimNextVariableReference::FromName("Bool", AnimationGraph), TestInt64) == EPropertyBagResult::Success, "FAnimationAnimNextRuntimeTest_Variables -> GetVariable failed");
 			UE_RETURN_ON_ERROR(TestInt64 == 1, "FAnimationAnimNextRuntimeTest_Variables -> Variable value did not match");
 
-			TestUInt64 = 0;
-			UE_RETURN_ON_ERROR(GraphInstance->GetVariable(FAnimNextVariableReference::FromName("Bool", AnimationGraph), TestUInt64) == EPropertyBagResult::Success, "FAnimationAnimNextRuntimeTest_Variables -> GetVariable failed");
-			UE_RETURN_ON_ERROR(TestUInt64 == 1, "FAnimationAnimNextRuntimeTest_Variables -> Variable value did not match");
-
 			// Byte
 			TestBool = false;
 			UE_RETURN_ON_ERROR(GraphInstance->GetVariable(FAnimNextVariableReference::FromName("Byte", AnimationGraph), TestBool) == EPropertyBagResult::Success, "FAnimationAnimNextRuntimeTest_Variables -> GetVariable failed");
@@ -769,18 +751,10 @@ bool FAnimationAnimNextRuntimeTest_Variables::RunTest(const FString& InParameter
 			UE_RETURN_ON_ERROR(GraphInstance->GetVariable(FAnimNextVariableReference::FromName("Byte", AnimationGraph), TestInt32) == EPropertyBagResult::Success, "FAnimationAnimNextRuntimeTest_Variables -> GetVariable failed");
 			UE_RETURN_ON_ERROR(TestInt32 == 42, "FAnimationAnimNextRuntimeTest_Variables -> Variable value did not match");
 
-			TestUInt32 = 0;
-			UE_RETURN_ON_ERROR(GraphInstance->GetVariable(FAnimNextVariableReference::FromName("Byte", AnimationGraph), TestUInt32) == EPropertyBagResult::Success, "FAnimationAnimNextRuntimeTest_Variables -> GetVariable failed");
-			UE_RETURN_ON_ERROR(TestUInt32 == 42, "FAnimationAnimNextRuntimeTest_Variables -> Variable value did not match");
-
 			TestInt64 = 0;
 			UE_RETURN_ON_ERROR(GraphInstance->GetVariable(FAnimNextVariableReference::FromName("Byte", AnimationGraph), TestInt64) == EPropertyBagResult::Success, "FAnimationAnimNextRuntimeTest_Variables -> GetVariable failed");
 			UE_RETURN_ON_ERROR(TestInt64 == 42, "FAnimationAnimNextRuntimeTest_Variables -> Variable value did not match");
 
-			TestUInt64 = 0;
-			UE_RETURN_ON_ERROR(GraphInstance->GetVariable(FAnimNextVariableReference::FromName("Byte", AnimationGraph), TestUInt64) == EPropertyBagResult::Success, "FAnimationAnimNextRuntimeTest_Variables -> GetVariable failed");
-			UE_RETURN_ON_ERROR(TestUInt64 == 42, "FAnimationAnimNextRuntimeTest_Variables -> Variable value did not match");
-
 			// Int32
 			TestBool = false;
 			UE_RETURN_ON_ERROR(GraphInstance->GetVariable(FAnimNextVariableReference::FromName("Int32", AnimationGraph), TestBool) == EPropertyBagResult::Success, "FAnimationAnimNextRuntimeTest_Variables -> GetVariable failed");
@@ -790,39 +764,10 @@ bool FAnimationAnimNextRuntimeTest_Variables::RunTest(const FString& InParameter
 			UE_RETURN_ON_ERROR(GraphInstance->GetVariable(FAnimNextVariableReference::FromName("Int32", AnimationGraph), TestByte) == EPropertyBagResult::Success, "FAnimationAnimNextRuntimeTest_Variables -> GetVariable failed");
 			UE_RETURN_ON_ERROR(TestByte == 202/*(uint8)-4679222*/, "FAnimationAnimNextRuntimeTest_Variables -> Variable value did not match");
 
-			TestUInt32 = 0;
-			UE_RETURN_ON_ERROR(GraphInstance->GetVariable(FAnimNextVariableReference::FromName("Int32", AnimationGraph), TestUInt32) == EPropertyBagResult::Success, "FAnimationAnimNextRuntimeTest_Variables -> GetVariable failed");
-			UE_RETURN_ON_ERROR(TestUInt32 == 4290288074/*(uint32)-4679222*/, "FAnimationAnimNextRuntimeTest_Variables -> Variable value did not match");
-
 			TestInt64 = 0;
 			UE_RETURN_ON_ERROR(GraphInstance->GetVariable(FAnimNextVariableReference::FromName("Int32", AnimationGraph), TestInt64) == EPropertyBagResult::Success, "FAnimationAnimNextRuntimeTest_Variables -> GetVariable failed");
 			UE_RETURN_ON_ERROR(TestInt64 == -4679222, "FAnimationAnimNextRuntimeTest_Variables -> Variable value did not match");
 
-			TestUInt64 = 0;
-			UE_RETURN_ON_ERROR(GraphInstance->GetVariable(FAnimNextVariableReference::FromName("Int32", AnimationGraph), TestUInt64) == EPropertyBagResult::Success, "FAnimationAnimNextRuntimeTest_Variables -> GetVariable failed");
-			UE_RETURN_ON_ERROR(TestUInt64 == 4290288074/*(uint64)-4679222*/, "FAnimationAnimNextRuntimeTest_Variables -> Variable value did not match");
-
-			// Uint32
-			TestBool = false;
-			UE_RETURN_ON_ERROR(GraphInstance->GetVariable(FAnimNextVariableReference::FromName("UInt32", AnimationGraph), TestBool) == EPropertyBagResult::Success, "FAnimationAnimNextRuntimeTest_Variables -> GetVariable failed");
-			UE_RETURN_ON_ERROR(TestBool == true, "FAnimationAnimNextRuntimeTest_Variables -> Variable value did not match");
-
-			TestByte = 0;
-			UE_RETURN_ON_ERROR(GraphInstance->GetVariable(FAnimNextVariableReference::FromName("UInt32", AnimationGraph), TestByte) == EPropertyBagResult::Success, "FAnimationAnimNextRuntimeTest_Variables -> GetVariable failed");
-			UE_RETURN_ON_ERROR(TestByte == 255/*(uint8)3415919103*/, "FAnimationAnimNextRuntimeTest_Variables -> Variable value did not match");
-
-			TestInt32 = 0;
-			UE_RETURN_ON_ERROR(GraphInstance->GetVariable(FAnimNextVariableReference::FromName("UInt32", AnimationGraph), TestInt32) == EPropertyBagResult::Success, "FAnimationAnimNextRuntimeTest_Variables -> GetVariable failed");
-			UE_RETURN_ON_ERROR(TestInt32 == -879048193/*(int32)3415919103*/, "FAnimationAnimNextRuntimeTest_Variables -> Variable value did not match");
-
-			TestInt64 = 0;
-			UE_RETURN_ON_ERROR(GraphInstance->GetVariable(FAnimNextVariableReference::FromName("UInt32", AnimationGraph), TestInt64) == EPropertyBagResult::Success, "FAnimationAnimNextRuntimeTest_Variables -> GetVariable failed");
-			UE_RETURN_ON_ERROR(TestInt64 == 3415919103, "FAnimationAnimNextRuntimeTest_Variables -> Variable value did not match");
-
-			TestUInt64 = 0;
-			UE_RETURN_ON_ERROR(GraphInstance->GetVariable(FAnimNextVariableReference::FromName("UInt32", AnimationGraph), TestUInt64) == EPropertyBagResult::Success, "FAnimationAnimNextRuntimeTest_Variables -> GetVariable failed");
-			UE_RETURN_ON_ERROR(TestUInt64 == 3415919103, "FAnimationAnimNextRuntimeTest_Variables -> Variable value did not match");
-
 			// Int64
 			TestBool = false;
 			UE_RETURN_ON_ERROR(GraphInstance->GetVariable(FAnimNextVariableReference::FromName("Int64", AnimationGraph), TestBool) == EPropertyBagResult::Success, "FAnimationAnimNextRuntimeTest_Variables -> GetVariable failed");
@@ -835,14 +780,6 @@ bool FAnimationAnimNextRuntimeTest_Variables::RunTest(const FString& InParameter
 			TestInt32 = 0;
 			UE_RETURN_ON_ERROR(GraphInstance->GetVariable(FAnimNextVariableReference::FromName("Int64", AnimationGraph), TestInt32) == EPropertyBagResult::Success, "FAnimationAnimNextRuntimeTest_Variables -> GetVariable failed");
 			UE_RETURN_ON_ERROR(TestInt32 == 879048191/*(int32)-3415919105*/, "FAnimationAnimNextRuntimeTest_Variables -> Variable value did not match");
-
-			TestUInt32 = 0;
-			UE_RETURN_ON_ERROR(GraphInstance->GetVariable(FAnimNextVariableReference::FromName("Int64", AnimationGraph), TestUInt32) == EPropertyBagResult::Success, "FAnimationAnimNextRuntimeTest_Variables -> GetVariable failed");
-			UE_RETURN_ON_ERROR(TestUInt32 == 879048191/*(uint32)-3415919105*/, "FAnimationAnimNextRuntimeTest_Variables -> Variable value did not match");
-
-			TestUInt64 = 0;
-			UE_RETURN_ON_ERROR(GraphInstance->GetVariable(FAnimNextVariableReference::FromName("Int64", AnimationGraph), TestUInt64) == EPropertyBagResult::Success, "FAnimationAnimNextRuntimeTest_Variables -> GetVariable failed");
-			UE_RETURN_ON_ERROR(TestUInt64 == 18446744070293632511ull/*(uint64)-3415919105*/, "FAnimationAnimNextRuntimeTest_Variables -> Variable value did not match");
 		}
 
 		// Float/double + conversions
@@ -874,10 +811,6 @@ bool FAnimationAnimNextRuntimeTest_Variables::RunTest(const FString& InParameter
 		UE_RETURN_ON_ERROR(GraphInstance->GetVariable(FAnimNextVariableReference::FromName("String", AnimationGraph), TestString) == EPropertyBagResult::Success, "FAnimationAnimNextRuntimeTest_Variables -> GetVariable failed");
 		UE_RETURN_ON_ERROR(TestString == "Test", "FAnimationAnimNextRuntimeTest_Variables -> Variable value did not match");
 
-		FText TestText;
-		UE_RETURN_ON_ERROR(GraphInstance->GetVariable(FAnimNextVariableReference::FromName("Text", AnimationGraph), TestText) == EPropertyBagResult::Success, "FAnimationAnimNextRuntimeTest_Variables -> GetVariable failed");
-		UE_RETURN_ON_ERROR(TestText.EqualTo(NSLOCTEXT("Tests", "Test", "Test")), "FAnimationAnimNextRuntimeTest_Variables -> Variable value did not match");
-
 		// Enums
 		{
 			EPropertyBagPropertyType TestEnum = EPropertyBagPropertyType::None;
diff --git a/Engine/Plugins/Experimental/UAF/UAFTests/Source/UAFTestSuite/Private/AnimNextEditorTest.cpp b/Engine/Plugins/Experimental/UAF/UAFTests/Source/UAFTestSuite/Private/AnimNextEditorTest.cpp
--- a/Engine/Plugins/Experimental/UAF/UAFTests/Source/UAFTestSuite/Private/AnimNextEditorTest.cpp
+++ b/Engine/Plugins/Experimental/UAF/UAFTests/Source/UAFTestSuite/Private/AnimNextEditorTest.cpp
@@ -98,7 +98,6 @@ bool FEditor_Variables::RunTest(const FString& InParameters)
 		TestVariableType(FAnimNextParamType::GetType<double>());
 		TestVariableType(FAnimNextParamType::GetType<FName>());
 		TestVariableType(FAnimNextParamType::GetType<FString>());
-		TestVariableType(FAnimNextParamType::GetType<FText>());
 		TestVariableType(FAnimNextParamType::GetType<EPropertyBagPropertyType>());
 		TestVariableType(FAnimNextParamType::GetType<FVector>());
 		TestVariableType(FAnimNextParamType::GetType<FQuat>());
@@ -218,4 +217,4 @@ bool FEditor_Variables_Python::RunTest(const FString& InParameters)
 
 }
 
-#endif	// WITH_DEV_AUTOMATION_TESTS && WITH_EDITOR
\ No newline at end of file
+#endif	// WITH_DEV_AUTOMATION_TESTS && WITH_EDITOR
diff --git a/Engine/Plugins/Experimental/UAF/UAFTests/Source/UAFTestSuite/Private/AnimNextVariablesTest.cpp b/Engine/Plugins/Experimental/UAF/UAFTests/Source/UAFTestSuite/Private/AnimNextVariablesTest.cpp
--- a/Engine/Plugins/Experimental/UAF/UAFTests/Source/UAFTestSuite/Private/AnimNextVariablesTest.cpp
+++ b/Engine/Plugins/Experimental/UAF/UAFTests/Source/UAFTestSuite/Private/AnimNextVariablesTest.cpp
@@ -138,7 +138,6 @@ bool FVariableTypesTest::RunTest(const FString& InParameters)
 	AddErrorIfFalse(FAnimNextParamType::GetType<double>().IsValid(), TEXT("double parameter is invalid."));
 	AddErrorIfFalse(FAnimNextParamType::GetType<FName>().IsValid(), TEXT("FName parameter is invalid."));
 	AddErrorIfFalse(FAnimNextParamType::GetType<FString>().IsValid(), TEXT("FString parameter is invalid."));
-	AddErrorIfFalse(FAnimNextParamType::GetType<FText>().IsValid(), TEXT("FText parameter is invalid."));
 	AddErrorIfFalse(FAnimNextParamType::GetType<EPropertyBagContainerType>().IsValid(), TEXT("Enum parameter is invalid."));
 	AddErrorIfFalse(FAnimNextParamType::GetType<FAnimNextParamType>().IsValid(), TEXT("Struct parameter is invalid."));
 	AddErrorIfFalse(FAnimNextParamType::GetType<FVector>().IsValid(), TEXT("Struct parameter is invalid."));
@@ -161,7 +160,6 @@ bool FVariableTypesTest::RunTest(const FString& InParameters)
 	AddErrorIfFalse(FAnimNextParamType::GetType<TArray<double>>().IsValid(), TEXT("double array parameter is invalid."));
 	AddErrorIfFalse(FAnimNextParamType::GetType<TArray<FName>>().IsValid(), TEXT("FName array parameter is invalid."));
 	AddErrorIfFalse(FAnimNextParamType::GetType<TArray<FString>>().IsValid(), TEXT("FString array parameter is invalid."));
-	AddErrorIfFalse(FAnimNextParamType::GetType<TArray<FText>>().IsValid(), TEXT("FText array parameter is invalid."));
 	AddErrorIfFalse(FAnimNextParamType::GetType<TArray<EPropertyBagContainerType>>().IsValid(), TEXT("Enum array parameter is invalid."));
 	AddErrorIfFalse(FAnimNextParamType::GetType<TArray<FAnimNextParamType>>().IsValid(), TEXT("Struct array parameter is invalid."));
 	AddErrorIfFalse(FAnimNextParamType::GetType<TArray<FVector>>().IsValid(), TEXT("Struct array parameter is invalid."));
@@ -188,7 +186,6 @@ bool FVariableTypesTest::RunTest(const FString& InParameters)
 	TEST_ANIMNEXT_PROPERTY(double, Double);
 	TEST_ANIMNEXT_PROPERTY(FName, Name);
 	TEST_ANIMNEXT_PROPERTY(FString, String);
-	TEST_ANIMNEXT_PROPERTY(FText, Text);
 	TEST_ANIMNEXT_PROPERTY(EPropertyBagContainerType, Enum);
 	TEST_ANIMNEXT_PROPERTY(FAnimNextParamType, Struct);
 	TEST_ANIMNEXT_PROPERTY(FVector, Vector);
@@ -208,7 +205,6 @@ bool FVariableTypesTest::RunTest(const FString& InParameters)
 	TEST_ANIMNEXT_PROPERTY_ARRAY(double, Double);
 	TEST_ANIMNEXT_PROPERTY_ARRAY(FName, Name);
 	TEST_ANIMNEXT_PROPERTY_ARRAY(FString, String);
-	TEST_ANIMNEXT_PROPERTY_ARRAY(FText, Text);
 	TEST_ANIMNEXT_PROPERTY_ARRAY(EPropertyBagContainerType, Enum);
 	TEST_ANIMNEXT_PROPERTY_ARRAY(FAnimNextParamType, Struct);
 	TEST_ANIMNEXT_PROPERTY_ARRAY(FVector, Vector);
@@ -478,15 +474,7 @@ namespace Private
 			Test.AddErrorIfFalse(GetResult == EPropertyBagResult::Success, *FString::Printf(TEXT("FVariables_Shared -> GetVariable call for %s failed."), InMemberName));
 
 			// Specialize comparisons for types that dont have operator==
-			if constexpr(std::is_same_v<Type, FText>)
-			{
-				Test.AddErrorIfFalse(TestValue.EqualTo(InValue), *FString::Printf(TEXT("FVariables_Shared -> Variable value did not match for %s."), InMemberName));
-			}
-			else if constexpr(std::is_same_v<Type, TArray<FText>>)
-			{
-				Test.AddErrorIfFalse(Algo::Compare(TestValue, InValue, [](const FText& A, const FText& B){ return A.EqualTo(B); }), *FString::Printf(TEXT("FVariables_Shared -> Variable value did not match for %s."), InMemberName));
-			}
-			else if constexpr(std::is_same_v<Type, FTransform>)
+			if constexpr(std::is_same_v<Type, FTransform>)
 			{
 				Test.AddErrorIfFalse(TestValue.Equals(InValue), *FString::Printf(TEXT("FVariables_Shared -> Variable value did not match for %s."), InMemberName));
 			}
@@ -530,7 +518,6 @@ bool FVariables_Shared::RunTest(const FString& InParameters)
 	Private::FTestReadWrite::Test<double>(*this, Instance, TEXT("Double"), 8756498576.0);
 	Private::FTestReadWrite::Test<FName>(*this, Instance, TEXT("Name"), FName("TestName"));
 	Private::FTestReadWrite::Test<FString>(*this, Instance, TEXT("String"), TEXT("TestString"));
-	Private::FTestReadWrite::Test<FText>(*this, Instance, TEXT("Text"), NSLOCTEXT("AnimNextTests", "TestText", "TestText"));
 	Private::FTestReadWrite::Test<EPropertyBagContainerType>(*this, Instance, TEXT("Enum"), EPropertyBagContainerType::Set);
 	Private::FTestReadWrite::Test<FAnimNextParamType>(*this, Instance, TEXT("Struct"), FAnimNextParamType::GetType<FAnimNextParamType>());
 	Private::FTestReadWrite::Test<FVector>(*this, Instance, TEXT("Vector"), FVector(1.0f, 2.0f, 3.0f));
@@ -550,7 +537,6 @@ bool FVariables_Shared::RunTest(const FString& InParameters)
 	Private::FTestReadWrite::Test<TArray<double>>(*this, Instance, TEXT("DoubleArray"), { 8756498576.0 });
 	Private::FTestReadWrite::Test<TArray<FName>>(*this, Instance, TEXT("NameArray"), { FName("TestName") });
 	Private::FTestReadWrite::Test<TArray<FString>>(*this, Instance, TEXT("StringArray"), { TEXT("TestString") });
-	Private::FTestReadWrite::Test<TArray<FText>>(*this, Instance, TEXT("TextArray"), { NSLOCTEXT("AnimNextTests", "TestText", "TestText") });
 	Private::FTestReadWrite::Test<TArray<EPropertyBagContainerType>>(*this, Instance, TEXT("EnumArray"), { EPropertyBagContainerType::Set });
 	Private::FTestReadWrite::Test<TArray<FAnimNextParamType>>(*this, Instance, TEXT("StructArray"), { FAnimNextParamType::GetType<FAnimNextParamType>() });
 	Private::FTestReadWrite::Test<TArray<FVector>>(*this, Instance, TEXT("VectorArray"), { FVector(1.0f, 2.0f, 3.0f) });
@@ -567,15 +553,7 @@ bool FVariables_Shared::RunTest(const FString& InParameters)
 		auto TestValue = [this]<typename T0>(const TCHAR* InMemberName, const auto& InValue, const T0& InComparison)
 		{
 			// Specialize comparisons for types that dont have operator==
-			if constexpr(std::is_same_v<T0, FText>)
-			{
-				AddErrorIfFalse(InComparison.EqualTo(InValue), *FString::Printf(TEXT("FVariables_Shared -> Set value does not match for %s."), InMemberName));
-			}
-			else if constexpr(std::is_same_v<T0, TArray<FText>>)
-			{
-				AddErrorIfFalse(Algo::Compare(InComparison, InValue, [](const FText& A, const FText& B){ return A.EqualTo(B); }), *FString::Printf(TEXT("FVariables_Shared -> Set value does not match for %s."), InMemberName));
-			}
-			else if constexpr(std::is_same_v<T0, FTransform>)
+			if constexpr(std::is_same_v<T0, FTransform>)
 			{
 				AddErrorIfFalse(InComparison.Equals(InValue), *FString::Printf(TEXT("FVariables_Shared -> Set value does not match for %s."), InMemberName));
 			}
@@ -599,7 +577,6 @@ bool FVariables_Shared::RunTest(const FString& InParameters)
 		TestValue(TEXT("Double"), InStruct.Double, 8756498576.0);
 		TestValue(TEXT("Name"), InStruct.Name, FName("TestName"));
 		TestValue(TEXT("String"), InStruct.String, FString("TestString"));
-		TestValue(TEXT("Text"), InStruct.Text, NSLOCTEXT("AnimNextTests", "TestText", "TestText"));
 		TestValue(TEXT("Enum"), InStruct.Enum, EPropertyBagContainerType::Set);
 		TestValue(TEXT("Struct"), InStruct.Struct, FAnimNextParamType::GetType<FAnimNextParamType>());
 		TestValue(TEXT("Vector"), InStruct.Vector, FVector(1.0f, 2.0f, 3.0f));
diff --git a/Engine/Plugins/Experimental/UAF/UAFTests/Source/UAFTestSuite/Private/AnimNextVariablesTest.h b/Engine/Plugins/Experimental/UAF/UAFTests/Source/UAFTestSuite/Private/AnimNextVariablesTest.h
--- a/Engine/Plugins/Experimental/UAF/UAFTests/Source/UAFTestSuite/Private/AnimNextVariablesTest.h
+++ b/Engine/Plugins/Experimental/UAF/UAFTests/Source/UAFTestSuite/Private/AnimNextVariablesTest.h
@@ -47,9 +47,6 @@ struct FAnimNextParamTypeTestStruct
 	UPROPERTY()
 	FString String;
 
-	UPROPERTY()
-	FText Text;
-
 	UPROPERTY()
 	EPropertyBagContainerType Enum = EPropertyBagContainerType::None;
 
@@ -105,9 +102,6 @@ struct FAnimNextParamTypeTestStruct
 	UPROPERTY()
 	TArray<FString> StringArray;
 
-	UPROPERTY()
-	TArray<FText> TextArray;
-
 	UPROPERTY()
 	TArray<EPropertyBagContainerType> EnumArray;
 
@@ -208,4 +202,4 @@ public:
 
 	UPROPERTY()
 	int32 ValueC = 12345;
-};
\ No newline at end of file
+};
