Add basic (python-based) UAF function execution test
#misc added script helper functionality: to run a RigVM part of a UAnimNextRigVMAsset and to create a transient UAnimNextRigVMAsset instance

#Synced-CL 48786226

diff --git a/Engine/Plugins/Experimental/UAF/UAF/Source/UAF/Public/AnimNextRigVMAsset.h b/Engine/Plugins/Experimental/UAF/UAF/Source/UAF/Public/AnimNextRigVMAsset.h
--- a/Engine/Plugins/Experimental/UAF/UAF/Source/UAF/Public/AnimNextRigVMAsset.h
+++ b/Engine/Plugins/Experimental/UAF/UAF/Source/UAF/Public/AnimNextRigVMAsset.h
@@ -77,6 +77,7 @@ protected:
 	friend UE::UAF::UncookedOnly::FCompilationScope;
 	friend UE::UAF::Tests::FSystemReference_Creation;
 	friend FUAFAssetInstance;
+	friend class UUAFTestBlueprintLibrary;
 
 	UE_API UAnimNextRigVMAsset(const FObjectInitializer& ObjectInitializer);
 
diff --git a/Engine/Plugins/Experimental/UAF/UAFTests/Source/UAFTests/Private/UAFTestBlueprintLibrary.h b/Engine/Plugins/Experimental/UAF/UAFTests/Source/UAFTests/Private/UAFTestBlueprintLibrary.h
--- a/Engine/Plugins/Experimental/UAF/UAFTests/Source/UAFTests/Private/UAFTestBlueprintLibrary.h
+++ b/Engine/Plugins/Experimental/UAF/UAFTests/Source/UAFTests/Private/UAFTestBlueprintLibrary.h
@@ -5,7 +5,11 @@
 #include "CoreMinimal.h"
 #include "Kismet/BlueprintFunctionLibrary.h"
 #include "AnimNextRigVMAssetEditorData.h"
+#include "AutomationTestModule.h"
 #include "UAFCompilationScope.h"
+#include "Factories/Factory.h"
+#include "Module/AnimNextModule_EditorData.h"
+#include "UObject/Class.h"
 #include "UAFTestBlueprintLibrary.generated.h"
 
 UCLASS()
@@ -177,4 +181,83 @@ public:
 
 		return nullptr;
 	}
+	
+	UFUNCTION(BlueprintCallable, Category="UAFTest", meta=(ScriptMethod))
+	static void ExecuteVM(UAnimNextModule* InModule, const FName& InEventName, bool& bExecutionResult, TArray<FString>& OutMessages)
+	{
+		bExecutionResult = false;
+		
+		if (InModule == nullptr)
+		{
+			UE_LOG(LogAnimation, Warning, TEXT("Invalid Asset in UUAFTestBlueprintLibrary::ExecuteVM"))
+			return;
+		}
+	
+		if (!InModule->SupportsEvent(InEventName))
+		{
+			UE_LOG(LogAnimation, Warning, TEXT("%s does not support event named %s"), *InModule->GetName(), *InEventName.ToString());
+			return;
+		}
+	
+		if (InModule->GetVM() == nullptr)
+		{
+			UE_LOG(LogAnimation, Warning, TEXT("%s does not contain a valid VM"), *InModule->GetName());
+			return;
+		}
+	
+		if (UAnimNextModule_EditorData* EditorData = UE::UAF::UncookedOnly::FUtils::GetEditorData<UAnimNextModule_EditorData>(InModule))
+		{
+			if (EditorData->bVMRecompilationRequired)
+			{
+				UE_LOG(LogAnimation, Warning, TEXT("%s requires compilation - so invoking now"), *InModule->GetName());
+				EditorData->RecompileVM();
+			}
+		
+			if (InModule->CompilationState == EAnimNextRigVMAssetState::CompiledWithErrors)
+			{
+				UE_LOG(LogAnimation, Warning, TEXT("%s compilation finished with errors"), *InModule->GetName());
+				return;
+			}
+		}
+	
+		FRigVMRuntimeSettings RuntimeSettings;
+		RuntimeSettings.SetLogFunction([&OutMessages, ModuleName = InModule->GetName()](const FRigVMLogSettings& InLogSettings, const FRigVMExecuteContext* InContext, const FString& Message)
+		{
+			OutMessages.Add(Message);
+		});
+		InModule->GetRigVMExtendedExecuteContext().SetRuntimeSettings(RuntimeSettings);
+
+		ERigVMExecuteResult Result = InModule->GetVM()->ExecuteVM(InModule->GetRigVMExtendedExecuteContext(), InEventName);
+	
+		InModule->GetRigVMExtendedExecuteContext().SetRuntimeSettings(FRigVMRuntimeSettings());
+
+		bExecutionResult = Result == ERigVMExecuteResult::Succeeded;
+	}
+	
+	UFUNCTION(BlueprintCallable, Category="UAFTest")
+	static UAnimNextRigVMAsset* CreateAsset(TSubclassOf<UFactory> InAssetFactoryClass, const FName& InName)
+	{
+		if(InAssetFactoryClass.Get() == nullptr)
+        {
+			UE_LOG(LogAnimation, Warning, TEXT("Invalid InAssetFactoryClass in UUAFTestBlueprintLibrary::CreateAsset"));
+			return nullptr;
+		}
+
+		UFactory* Factory = NewObject<UFactory>(GetTransientPackage(), InAssetFactoryClass.Get());
+		if(Factory == nullptr)
+        {
+			UE_LOG(LogAnimation, Warning, TEXT("Failed to create Factory of class %s"), *InAssetFactoryClass.Get()->GetName());
+			return nullptr;
+        }
+		
+		if (!Factory->GetSupportedClass()->IsChildOf(UAnimNextRigVMAsset::StaticClass()))
+		{
+			UE_LOG(LogAnimation, Warning, TEXT("Factory %s provided to UUAFTestBlueprintLibrary::CreateAsset does not support a UAnimNextRigVMAsset-based asset"), *InAssetFactoryClass.Get()->GetName());
+			return nullptr;
+		}
+		
+		UAnimNextRigVMAsset* Asset = Cast<UAnimNextRigVMAsset>(Factory->FactoryCreateNew(Factory->GetSupportedClass(), GetTransientPackage(), InName, RF_Transient, nullptr, nullptr, NAME_None));
+		
+		return Asset;
+	}
 };
