Move and rename UAF ValueRuntime related files

diff --git a/Engine/Plugins/Experimental/UAF/UAF/Source/UAF/Public/UAF/ValueRuntime/ValueRuntimeRegistry.h b/Engine/Plugins/Experimental/UAF/UAF/Source/UAF/Public/UAF/ValueRuntime/ValueRuntimeRegistry.h
new file mode 100644
--- /dev/null
+++ b/Engine/Plugins/Experimental/UAF/UAF/Source/UAF/Public/UAF/ValueRuntime/ValueRuntimeRegistry.h
@@ -0,0 +1,159 @@
+// Copyright Epic Games, Inc. All Rights Reserved.
+
+#pragma once
+
+#include "CoreMinimal.h"
+#include "UAF/ValueRuntime/AttributeMap.h"
+#include "UAF/ValueRuntime/AttributeSetMap.h"
+#include "UAF/ValueRuntime/Transformers/Transformer.h"
+
+#define UE_API UAF_API
+
+namespace UE::UAF
+{
+	// A global registry to store value runtime metadata
+	class FValueRuntimeRegistry
+	{
+	public:
+		UE_NONCOPYABLE(FValueRuntimeRegistry);	// We disallow copy/move semantics
+
+		using SetMapConstructorFun = TFunction<FAttributeSetMap*(const FAttributeSetMap::FConstructArgs& Args)>;
+		using MapConstructorFun = TFunction<FAttributeMap*(FReallocFun ReallocFun)>;
+
+		// Returns the set map registry instance
+		[[nodiscard]] static UE_API FValueRuntimeRegistry& Get();
+
+		// Constructs a set map from its UStruct value type
+		[[nodiscard]] UE_API FAttributeSetMap* ConstructSetMap(const FAttributeSetMap::FConstructArgs& Args) const;
+
+		// Constructs a name/value map from its UStruct value type
+		[[nodiscard]] UE_API FAttributeMap* ConstructMap(UScriptStruct* ValueType, FReallocFun ReallocFun) const;
+
+		// Returns the map of known attribute value transformers
+		[[nodiscard]] UE_API FAttributeValueTransformerMapPtr GetTransformerMap() const;
+
+		// Registers a set map of the specified value type
+		template<class ValueType>
+		void RegisterSetMapInitializer();
+
+		// Registers a name/value map of the specified value type
+		template<class ValueType>
+		void RegisterMapInitializer();
+
+		template<class TransformerSpecializationType>
+		bool RegisterAttributeValueTransformer();
+
+	private:
+		UE_API FValueRuntimeRegistry();
+		UE_API void RegisterSetMapInitializer(UScriptStruct* ValueType, SetMapConstructorFun Initializer);
+		UE_API void RegisterMapInitializer(UScriptStruct* ValueType, MapConstructorFun Initializer);
+		UE_API bool RegisterAttributeValueSetMapTransformer(FName TransformerName, UScriptStruct* ValueType, FRawTransformerFunc TransformerFunc);
+		UE_API bool RegisterAttributeValueMapTransformer(FName TransformerName, UScriptStruct* ValueType, FRawTransformerFunc TransformerFunc);
+
+		// Initialize/destroy the global registry
+		static UE_API void Init();
+		static UE_API void Destroy();
+
+		// A map of ValueType -> Set Map constructors
+		// This is used for indirect construction when we have a UScriptStruct value type
+		TMap<UScriptStruct*, SetMapConstructorFun> SetMapInitializers;
+		mutable FRWLock SetMapInitializerLock;
+
+		// A map of ValueType -> Map constructors
+		// This is used for indirect construction when we have a UScriptStruct value type
+		TMap<UScriptStruct*, MapConstructorFun> MapInitializers;
+		mutable FRWLock MapInitializerLock;
+
+		// A shared pointer to a map of transformer name -> transformer list (a transformer per value type)
+		FAttributeValueTransformerMapPtr TransformerMap;
+		mutable FRWLock TransformerMapLock;
+
+		friend class FAnimNextModuleImpl;
+	};
+
+	//////////////////////////////////////////////////////////////////////////
+	// Implementation
+
+	template<class ValueType>
+	inline void FValueRuntimeRegistry::RegisterSetMapInitializer()
+	{
+		RegisterSetMapInitializer(ValueType::StaticStruct(), [](const FAttributeSetMap::FConstructArgs& Args)
+			{
+				typename TAttributeSetMap<ValueType>::FConstructArgs TypedArgs(Args.TypedSet, Args.ValueType, Args.ReallocFun);
+
+				return MakeAttributeSetMap<ValueType>(Args);
+			});
+	}
+
+	template<class ValueType>
+	inline void FValueRuntimeRegistry::RegisterMapInitializer()
+	{
+		RegisterMapInitializer(ValueType::StaticStruct(), [](FReallocFun ReallocFun)
+			{
+				return MakeAttributeMap<ValueType>(ReallocFun);
+			});
+	}
+
+	template<class TransformerSpecializationType>
+	inline bool FValueRuntimeRegistry::RegisterAttributeValueTransformer()
+	{
+		using TransformerCppType = typename TransformerSpecializationType::TransformerType;
+		using ValueType = typename TransformerSpecializationType::ValueType;
+
+		static_assert(std::is_base_of_v<FAttributeValueTransformer, TransformerCppType>, "Specified transformer must derive from FAttributeValueTransformer and provide the necessary API");
+
+		const FName TransformerName = TransformerCppType::TransformerName;
+		UScriptStruct* ValueTypeStruct = ValueType::StaticStruct();
+
+		bool bSuccess = false;
+
+		if constexpr (requires { &TransformerSpecializationType::TransformSetMap; })
+		{
+			// Signature must match our transformer
+			using FTransformSetMapFunc = typename TransformerCppType::FTransformSetMapFunc;
+			FTransformSetMapFunc SpecializedTransformFunc = &TransformerSpecializationType::TransformSetMap;
+
+#ifdef _MSC_VER
+	#pragma warning(push)
+	#pragma warning(disable : 4191) // unsafe conversion from 'function pointer type A' to 'function pointer type B'
+#endif
+
+			// Cast to a raw dummy function pointer type that we exclusively use for storage
+			// The transformer is responsible for casting it back to the original type it specifies
+			auto RawSpecializedTransformFunc = reinterpret_cast<FRawTransformerFunc>(SpecializedTransformFunc);
+
+#ifdef _MSC_VER
+	#pragma warning(pop)
+#endif
+
+			bSuccess |= RegisterAttributeValueSetMapTransformer(TransformerName, ValueTypeStruct, RawSpecializedTransformFunc);
+		}
+
+		if constexpr (requires { &TransformerSpecializationType::TransformMap; })
+		{
+			// Signature must match our transformer
+			using FTransformMapFunc = typename TransformerCppType::FTransformMapFunc;
+			FTransformMapFunc SpecializedTransformFunc = &TransformerSpecializationType::TransformMap;
+
+#ifdef _MSC_VER
+	#pragma warning(push)
+	#pragma warning(disable : 4191) // unsafe conversion from 'function pointer type A' to 'function pointer type B'
+#endif
+
+			// Cast to a raw dummy function pointer type that we exclusively use for storage
+			// The transformer is responsible for casting it back to the original type it specifies
+			auto RawSpecializedTransformFunc = reinterpret_cast<FRawTransformerFunc>(SpecializedTransformFunc);
+
+#ifdef _MSC_VER
+	#pragma warning(pop)
+#endif
+
+			bSuccess |= RegisterAttributeValueMapTransformer(TransformerName, ValueTypeStruct, RawSpecializedTransformFunc);
+		}
+
+		// Returns true if we registered with at least one transformer type
+		return bSuccess;
+	}
+}
+
+#undef UE_API
diff --git a/Engine/Plugins/Experimental/UAF/UAFAnimGraph/Source/UAFAnimGraph/Private/EvaluationVM/EvaluationVM.cpp b/Engine/Plugins/Experimental/UAF/UAFAnimGraph/Source/UAFAnimGraph/Private/EvaluationVM/EvaluationVM.cpp
--- a/Engine/Plugins/Experimental/UAF/UAFAnimGraph/Source/UAFAnimGraph/Private/EvaluationVM/EvaluationVM.cpp
+++ b/Engine/Plugins/Experimental/UAF/UAFAnimGraph/Source/UAFAnimGraph/Private/EvaluationVM/EvaluationVM.cpp
@@ -35,32 +35,32 @@ namespace UE::UAF
 		, EvaluationFlags(InEvaluationFlags)
 	{
 		// Push an invalid context
-		EvaluationContextStack.Push(FAttributeEvaluationContext());
+		EvaluationContextStack.Push(FEvaluationTaskContext());
 	}
 
 	FEvaluationVM::FEvaluationVM(EEvaluationFlags InEvaluationFlags, const UAbstractSkeletonSetBinding& InSetBinding, FName InSetName, int32 InCurrentLOD)
-		: AttributeValueTransformerMap(FAttributeRuntimeRegistry::Get().GetTransformerMap())
+		: AttributeValueTransformerMap(FValueRuntimeRegistry::Get().GetTransformerMap())
 		, CurrentLOD(InCurrentLOD)
 		, EvaluationFlags(InEvaluationFlags)
 	{
-		EvaluationContextStack.Push(FAttributeEvaluationContext(InSetBinding, InSetName));
+		EvaluationContextStack.Push(FEvaluationTaskContext(InSetBinding, InSetName));
 	}
 
 	// REMOVE ME: Exclusively for testing that before/after yields the same results
 	FEvaluationVM::FEvaluationVM(EEvaluationFlags InEvaluationFlags, const FReferencePose& InReferencePose, const UAbstractSkeletonSetBinding* InSetBinding, FName InSetName, int32 InCurrentLOD)
 		: ReferencePose(&InReferencePose)
-		, AttributeValueTransformerMap(FAttributeRuntimeRegistry::Get().GetTransformerMap())
+		, AttributeValueTransformerMap(FValueRuntimeRegistry::Get().GetTransformerMap())
 		, CurrentLOD(InCurrentLOD)
 		, EvaluationFlags(InEvaluationFlags)
 	{
 		if (InSetBinding != nullptr)
 		{
-			EvaluationContextStack.Push(FAttributeEvaluationContext(*InSetBinding, InSetName));
+			EvaluationContextStack.Push(FEvaluationTaskContext(*InSetBinding, InSetName));
 		}
 		else
 		{
 			// Push an invalid context
-			EvaluationContextStack.Push(FAttributeEvaluationContext());
+			EvaluationContextStack.Push(FEvaluationTaskContext());
 		}
 	}
 
@@ -119,7 +119,7 @@ namespace UE::UAF
 		return Keyframe;
 	}
 
-	const FAttributeEvaluationContext& FEvaluationVM::GetActiveEvaluationContext() const
+	const FEvaluationTaskContext& FEvaluationVM::GetActiveEvaluationContext() const
 	{
 		return EvaluationContextStack.Last();
 	}
@@ -131,7 +131,7 @@ namespace UE::UAF
 
 	bool FEvaluationVM::PushEvaluationContext(const UAbstractSkeletonSetBinding& SetBinding, FName SetName)
 	{
-		FAttributeEvaluationContext NewContext(SetBinding, SetName);
+		FEvaluationTaskContext NewContext(SetBinding, SetName);
 		if (NewContext.IsValid())
 		{
 			EvaluationContextStack.Push(MoveTemp(NewContext));
@@ -145,7 +145,7 @@ namespace UE::UAF
 	bool FEvaluationVM::PushEvaluationContext(FName SetName)
 	{
 		// Copy the currently active context
-		FAttributeEvaluationContext NewContext = EvaluationContextStack.Last();
+		FEvaluationTaskContext NewContext = EvaluationContextStack.Last();
 		if (NewContext.SetNamedSet(SetName))
 		{
 			EvaluationContextStack.Push(MoveTemp(NewContext));
diff --git a/Engine/Plugins/Experimental/UAF/UAFAnimGraph/Source/UAFAnimGraph/Private/EvaluationVM/Tasks/ApplyAdditiveKeyframe.cpp b/Engine/Plugins/Experimental/UAF/UAFAnimGraph/Source/UAFAnimGraph/Private/EvaluationVM/Tasks/ApplyAdditiveKeyframe.cpp
--- a/Engine/Plugins/Experimental/UAF/UAFAnimGraph/Source/UAFAnimGraph/Private/EvaluationVM/Tasks/ApplyAdditiveKeyframe.cpp
+++ b/Engine/Plugins/Experimental/UAF/UAFAnimGraph/Source/UAFAnimGraph/Private/EvaluationVM/Tasks/ApplyAdditiveKeyframe.cpp
@@ -5,9 +5,9 @@
 #include "EvaluationVM/EvaluationVM.h"
 #include "EvaluationVM/KeyframeState.h"
 #include "TransformArrayOperations.h"
-#include "UAF/Attributes/AttributeCollection.h"
-#include "UAF/Attributes/Transformers/AdditiveSpace.h"
-#include "UAF/Attributes/Transformers/BoneSpace.h"
+#include "UAF/ValueRuntime/AttributeCollection.h"
+#include "UAF/ValueRuntime/Transformers/AdditiveSpace.h"
+#include "UAF/ValueRuntime/Transformers/BoneSpace.h"
 
 #include UE_INLINE_GENERATED_CPP_BY_NAME(ApplyAdditiveKeyframe)
 
diff --git a/Engine/Plugins/Experimental/UAF/UAFAnimGraph/Source/UAFAnimGraph/Private/EvaluationVM/Tasks/ApplyAdditiveKeyframePerBone.cpp b/Engine/Plugins/Experimental/UAF/UAFAnimGraph/Source/UAFAnimGraph/Private/EvaluationVM/Tasks/ApplyAdditiveKeyframePerBone.cpp
--- a/Engine/Plugins/Experimental/UAF/UAFAnimGraph/Source/UAFAnimGraph/Private/EvaluationVM/Tasks/ApplyAdditiveKeyframePerBone.cpp
+++ b/Engine/Plugins/Experimental/UAF/UAFAnimGraph/Source/UAFAnimGraph/Private/EvaluationVM/Tasks/ApplyAdditiveKeyframePerBone.cpp
@@ -6,9 +6,9 @@
 #include "EvaluationVM/KeyframeState.h"
 #include "TransformArrayOperations.h"
 #include "Animation/AnimationAsset.h"
-#include "UAF/Attributes/AttributeCollection.h"
-#include "UAF/Attributes/Transformers/AdditiveSpace.h"
-#include "UAF/Attributes/Transformers/BoneSpace.h"
+#include "UAF/ValueRuntime/AttributeCollection.h"
+#include "UAF/ValueRuntime/Transformers/AdditiveSpace.h"
+#include "UAF/ValueRuntime/Transformers/BoneSpace.h"
 
 #include UE_INLINE_GENERATED_CPP_BY_NAME(ApplyAdditiveKeyframePerBone)
 
diff --git a/Engine/Plugins/Experimental/UAF/UAFAnimGraph/Source/UAFAnimGraph/Private/EvaluationVM/Tasks/BlendKeyframes.cpp b/Engine/Plugins/Experimental/UAF/UAFAnimGraph/Source/UAFAnimGraph/Private/EvaluationVM/Tasks/BlendKeyframes.cpp
--- a/Engine/Plugins/Experimental/UAF/UAFAnimGraph/Source/UAFAnimGraph/Private/EvaluationVM/Tasks/BlendKeyframes.cpp
+++ b/Engine/Plugins/Experimental/UAF/UAFAnimGraph/Source/UAFAnimGraph/Private/EvaluationVM/Tasks/BlendKeyframes.cpp
@@ -6,11 +6,11 @@
 #include "EvaluationVM/KeyframeState.h"
 #include "TransformArrayOperations.h"
 #include "Animation/InputScaleBias.h"
-#include "UAF/Attributes/AttributeCollection.h"
-#include "UAF/Attributes/Transformers/Accumulate.h"
-#include "UAF/Attributes/Transformers/Interpolate.h"
-#include "UAF/Attributes/Transformers/Overwrite.h"
-#include "UAF/Attributes/Transformers/Sanitize.h"
+#include "UAF/ValueRuntime/AttributeCollection.h"
+#include "UAF/ValueRuntime/Transformers/Accumulate.h"
+#include "UAF/ValueRuntime/Transformers/Interpolate.h"
+#include "UAF/ValueRuntime/Transformers/Overwrite.h"
+#include "UAF/ValueRuntime/Transformers/Sanitize.h"
 
 #include UE_INLINE_GENERATED_CPP_BY_NAME(BlendKeyframes)
 
diff --git a/Engine/Plugins/Experimental/UAF/UAFAnimGraph/Source/UAFAnimGraph/Private/EvaluationVM/Tasks/BlendKeyframesPerBone.cpp b/Engine/Plugins/Experimental/UAF/UAFAnimGraph/Source/UAFAnimGraph/Private/EvaluationVM/Tasks/BlendKeyframesPerBone.cpp
--- a/Engine/Plugins/Experimental/UAF/UAFAnimGraph/Source/UAFAnimGraph/Private/EvaluationVM/Tasks/BlendKeyframesPerBone.cpp
+++ b/Engine/Plugins/Experimental/UAF/UAFAnimGraph/Source/UAFAnimGraph/Private/EvaluationVM/Tasks/BlendKeyframesPerBone.cpp
@@ -11,10 +11,10 @@
 #include "EvaluationVM/EvaluationVM.h"
 #include "EvaluationVM/KeyframeState.h"
 #include "TransformArrayOperations.h"
-#include "UAF/Attributes/Transformers/Accumulate.h"
-#include "UAF/Attributes/Transformers/Layer.h"
-#include "UAF/Attributes/Transformers/Overwrite.h"
-#include "UAF/Attributes/Transformers/Sanitize.h"
+#include "UAF/ValueRuntime/Transformers/Accumulate.h"
+#include "UAF/ValueRuntime/Transformers/Layer.h"
+#include "UAF/ValueRuntime/Transformers/Overwrite.h"
+#include "UAF/ValueRuntime/Transformers/Sanitize.h"
 
 #include UE_INLINE_GENERATED_CPP_BY_NAME(BlendKeyframesPerBone)
 
diff --git a/Engine/Plugins/Experimental/UAF/UAFAnimGraph/Source/UAFAnimGraph/Private/EvaluationVM/Tasks/ConvertRotationsToLocalSpace.cpp b/Engine/Plugins/Experimental/UAF/UAFAnimGraph/Source/UAFAnimGraph/Private/EvaluationVM/Tasks/ConvertRotationsToLocalSpace.cpp
--- a/Engine/Plugins/Experimental/UAF/UAFAnimGraph/Source/UAFAnimGraph/Private/EvaluationVM/Tasks/ConvertRotationsToLocalSpace.cpp
+++ b/Engine/Plugins/Experimental/UAF/UAFAnimGraph/Source/UAFAnimGraph/Private/EvaluationVM/Tasks/ConvertRotationsToLocalSpace.cpp
@@ -4,7 +4,7 @@
 
 #include "EvaluationVM/EvaluationVM.h"
 #include "EvaluationVM/Tasks/ConvertRotationSpaceDefs.h"
-#include "UAF/Attributes/Transformers/BoneSpace.h"
+#include "UAF/ValueRuntime/Transformers/BoneSpace.h"
 
 #if UE_ANIM_CONVERT_ROTATION_SPACE_TASKS_ISPC
 #include "ConvertRotationSpace.ispc.generated.h"
diff --git a/Engine/Plugins/Experimental/UAF/UAFAnimGraph/Source/UAFAnimGraph/Private/EvaluationVM/Tasks/ConvertRotationsToMeshSpace.cpp b/Engine/Plugins/Experimental/UAF/UAFAnimGraph/Source/UAFAnimGraph/Private/EvaluationVM/Tasks/ConvertRotationsToMeshSpace.cpp
--- a/Engine/Plugins/Experimental/UAF/UAFAnimGraph/Source/UAFAnimGraph/Private/EvaluationVM/Tasks/ConvertRotationsToMeshSpace.cpp
+++ b/Engine/Plugins/Experimental/UAF/UAFAnimGraph/Source/UAFAnimGraph/Private/EvaluationVM/Tasks/ConvertRotationsToMeshSpace.cpp
@@ -4,7 +4,7 @@
 
 #include "EvaluationVM/EvaluationVM.h"
 #include "EvaluationVM/Tasks/ConvertRotationSpaceDefs.h"
-#include "UAF/Attributes/Transformers/BoneSpace.h"
+#include "UAF/ValueRuntime/Transformers/BoneSpace.h"
 
 #if UE_ANIM_CONVERT_ROTATION_SPACE_TASKS_ISPC
 #include "ConvertRotationSpace.ispc.generated.h"
diff --git a/Engine/Plugins/Experimental/UAF/UAFAnimGraph/Source/UAFAnimGraph/Private/EvaluationVM/Tasks/NormalizeRotations.cpp b/Engine/Plugins/Experimental/UAF/UAFAnimGraph/Source/UAFAnimGraph/Private/EvaluationVM/Tasks/NormalizeRotations.cpp
--- a/Engine/Plugins/Experimental/UAF/UAFAnimGraph/Source/UAFAnimGraph/Private/EvaluationVM/Tasks/NormalizeRotations.cpp
+++ b/Engine/Plugins/Experimental/UAF/UAFAnimGraph/Source/UAFAnimGraph/Private/EvaluationVM/Tasks/NormalizeRotations.cpp
@@ -5,8 +5,8 @@
 #include "EvaluationVM/EvaluationVM.h"
 #include "EvaluationVM/KeyframeState.h"
 #include "TransformArrayOperations.h"
-#include "UAF/Attributes/AttributeCollection.h"
-#include "UAF/Attributes/Transformers/Sanitize.h"
+#include "UAF/ValueRuntime/AttributeCollection.h"
+#include "UAF/ValueRuntime/Transformers/Sanitize.h"
 
 #include UE_INLINE_GENERATED_CPP_BY_NAME(NormalizeRotations)
 
diff --git a/Engine/Plugins/Experimental/UAF/UAFAnimGraph/Source/UAFAnimGraph/Private/EvaluationVM/Tasks/PushAnimSequenceKeyframe.cpp b/Engine/Plugins/Experimental/UAF/UAFAnimGraph/Source/UAFAnimGraph/Private/EvaluationVM/Tasks/PushAnimSequenceKeyframe.cpp
--- a/Engine/Plugins/Experimental/UAF/UAFAnimGraph/Source/UAFAnimGraph/Private/EvaluationVM/Tasks/PushAnimSequenceKeyframe.cpp
+++ b/Engine/Plugins/Experimental/UAF/UAFAnimGraph/Source/UAFAnimGraph/Private/EvaluationVM/Tasks/PushAnimSequenceKeyframe.cpp
@@ -5,7 +5,7 @@
 #include "Animation/AnimSequence.h"
 #include "Animation/AnimRootMotionProvider.h"
 #include "BonePose.h"
-#include "DecompressionTools.h"
+#include "EvaluationVM/Tasks/DecompressionTools.h"
 #include "EvaluationVM/EvaluationVM.h"
 #include "EvaluationVM/KeyframeState.h"
 
@@ -34,7 +34,7 @@ void FAnimNextAnimSequenceKeyframeTask::Execute(UE::UAF::FEvaluationVM& VM) cons
 {
 	using namespace UE::UAF;
 	
-	const FAttributeEvaluationContext& VMContext = VM.GetActiveEvaluationContext();
+	const FEvaluationTaskContext& VMContext = VM.GetActiveEvaluationContext();
 
 	if(const UAnimSequence* AnimSequencePtr = AnimSequence.Get())
 	{
diff --git a/Engine/Plugins/Experimental/UAF/UAFAnimGraph/Source/UAFAnimGraph/Private/TraitInterfaces/IEvaluate.cpp b/Engine/Plugins/Experimental/UAF/UAFAnimGraph/Source/UAFAnimGraph/Private/TraitInterfaces/IEvaluate.cpp
--- a/Engine/Plugins/Experimental/UAF/UAFAnimGraph/Source/UAFAnimGraph/Private/TraitInterfaces/IEvaluate.cpp
+++ b/Engine/Plugins/Experimental/UAF/UAFAnimGraph/Source/UAFAnimGraph/Private/TraitInterfaces/IEvaluate.cpp
@@ -12,8 +12,8 @@
 #include "Graph/AnimNext_LODPose.h"
 #include "Graph/TraceAnimNextGraphInstances.h"
 
-#include "UAF/Attributes/AttributeCollection.h"
 #include "UAF/Attributes/EngineAttributes.h"
+#include "UAF/ValueRuntime/AttributeCollection.h"
 
 namespace UE::UAF
 {
diff --git a/Engine/Plugins/Experimental/UAF/UAFAnimGraph/Source/UAFAnimGraph/Public/EvaluationVM/EvaluationVM.h b/Engine/Plugins/Experimental/UAF/UAFAnimGraph/Source/UAFAnimGraph/Public/EvaluationVM/EvaluationVM.h
--- a/Engine/Plugins/Experimental/UAF/UAFAnimGraph/Source/UAFAnimGraph/Public/EvaluationVM/EvaluationVM.h
+++ b/Engine/Plugins/Experimental/UAF/UAFAnimGraph/Source/UAFAnimGraph/Public/EvaluationVM/EvaluationVM.h
@@ -7,15 +7,15 @@
 #include "ConstExprUID.h"
 #include "Animation/AnimCurveFilter.h"
 #include "EvaluationVM/EvaluationFlags.h"
+#include "EvaluationVM/EvaluationTaskContext.h"
 #include "EvaluationVM/KeyframeState.h"
 #include "Hash/Fnv.h"
 #include "ReferencePose.h"
 #include "Misc/Crc.h"
 #include "UAF/AbstractSkeleton/AbstractSkeletonSetBinding.h"
-#include "UAF/Attributes/AttributeCollection.h"
-#include "UAF/Attributes/AttributeEvaluationContext.h"
-#include "UAF/Attributes/AttributeRuntimeRegistry.h"
-#include "UAF/Attributes/Transformers/Transformer.h"
+#include "UAF/ValueRuntime/AttributeCollection.h"
+#include "UAF/ValueRuntime/Transformers/Transformer.h"
+#include "UAF/ValueRuntime/ValueRuntimeRegistry.h"
 
 #define UE_API UAFANIMGRAPH_API
 
@@ -198,7 +198,7 @@ namespace UE::UAF
 		[[nodiscard]] UE_API FKeyframeState MakeUninitializedKeyframe(bool bAdditiveKeyframe) const;
 
 		// Returns the evaluation context that controls what state we produce
-		[[nodiscard]] UE_API const FAttributeEvaluationContext& GetActiveEvaluationContext() const;
+		[[nodiscard]] UE_API const FEvaluationTaskContext& GetActiveEvaluationContext() const;
 
 		// Returns the named set we are currently evaluating with
 		[[nodiscard]] UE_API const FAttributeNamedSetPtr& GetActiveNamedSet() const;
@@ -242,7 +242,7 @@ namespace UE::UAF
 		const FReferencePose* ReferencePose = nullptr;
 
 		// The internal context that controls what state we are producing
-		TArray<FAttributeEvaluationContext> EvaluationContextStack;
+		TArray<FEvaluationTaskContext> EvaluationContextStack;
 
 		// Shared map of attribute value transformers
 		FAttributeValueTransformerMapPtr AttributeValueTransformerMap;
diff --git a/Engine/Plugins/Experimental/UAF/UAFTests/Source/UAFTestSuite/Private/UAF/ValueRuntime/AttributeCollectionTest.cpp b/Engine/Plugins/Experimental/UAF/UAFTests/Source/UAFTestSuite/Private/UAF/ValueRuntime/AttributeCollectionTest.cpp
new file mode 100644
--- /dev/null
+++ b/Engine/Plugins/Experimental/UAF/UAFTests/Source/UAFTestSuite/Private/UAF/ValueRuntime/AttributeCollectionTest.cpp
@@ -0,0 +1,500 @@
+// Copyright Epic Games, Inc. All Rights Reserved.
+
+#include "CoreMinimal.h"
+#include "Misc/AutomationTest.h"
+
+#if WITH_DEV_AUTOMATION_TESTS
+
+#include "AnimNextTest.h"
+#include "Animation/BuiltInAttributeTypes.h"
+#include "UAF/AbstractSkeleton/SetBindingFactory.h"
+#include "UAF/AbstractSkeleton/SetCollectionFactory.h"
+#include "UAF/AbstractSkeleton/AbstractSkeletonSetBinding.h"
+#include "UAF/AbstractSkeleton/AbstractSkeletonSetCollection.h"
+#include "UAF/Attributes/AttributeBindingData.h"
+#include "UAF/Attributes/EngineAttributes.h"
+#include "UAF/ValueRuntime/AttributeCollection.h"
+#include "UAF/ValueRuntime/AttributeSetMap.h"
+#include "UAF/ValueRuntime/IteratorTestUtils.h"
+#include "UObject/UObjectGlobals.h"
+#include "UObject/Package.h"
+
+namespace UE::UAF::Tests
+{
+	IMPLEMENT_SIMPLE_AUTOMATION_TEST(FAttributeCollectionTest, "Animation.UAF.Attributes.AttributeCollection", EAutomationTestFlags::EditorContext | EAutomationTestFlags::EngineFilter)
+
+	bool FAttributeCollectionTest::RunTest(const FString& InParameters)
+	{
+		ON_SCOPE_EXIT{ FUtils::CleanupAfterTests(); };
+
+		FAttributeBindingDataPtr BindingData;
+		{
+			UFactory* BindingFactory = NewObject<UFactory>(GetTransientPackage(), UAbstractSkeletonSetBindingFactory::StaticClass());
+			UE_RETURN_ON_ERROR(BindingFactory != nullptr, "Failed to create factory");
+
+			UAbstractSkeletonSetBinding* Binding = CastChecked<UAbstractSkeletonSetBinding>(BindingFactory->FactoryCreateNew(UAbstractSkeletonSetBinding::StaticClass(), GetTransientPackage(), TEXT("TestBindingAsset"), RF_Transient, nullptr, nullptr, NAME_None));
+			UE_RETURN_ON_ERROR(Binding != nullptr, "Failed to create binding asset");
+
+			USkeleton* Skeleton = NewObject<USkeleton>(GetTransientPackage(), USkeleton::StaticClass(), TEXT("TestSkeletonAsset"), RF_Transient);
+			UE_RETURN_ON_ERROR(Skeleton != nullptr, "Failed to create skeleton asset");
+
+			{
+				// Not terribly clean, we cast away the 'const' to modify the skeleton
+				FReferenceSkeleton& RefSkeleton = const_cast<FReferenceSkeleton&>(Skeleton->GetReferenceSkeleton());
+				FReferenceSkeletonModifier SkeletonModifier(RefSkeleton, Skeleton);
+
+				SkeletonModifier.Add(FMeshBoneInfo(TEXT("Root"), TEXT("Root"), INDEX_NONE), FTransform::Identity);
+				SkeletonModifier.Add(FMeshBoneInfo(TEXT("Pelvis"), TEXT("Root"), 0), FTransform::Identity);
+				SkeletonModifier.Add(FMeshBoneInfo(TEXT("Spine_01"), TEXT("Spine_01"), 1), FTransform::Identity);
+				SkeletonModifier.Add(FMeshBoneInfo(TEXT("Spine_02"), TEXT("Spine_02"), 2), FTransform::Identity);
+				SkeletonModifier.Add(FMeshBoneInfo(TEXT("Head"), TEXT("Head"), 3), FTransform::Identity);
+				SkeletonModifier.Add(FMeshBoneInfo(TEXT("LeftArm"), TEXT("LeftArm"), 3), FTransform::Identity);
+				SkeletonModifier.Add(FMeshBoneInfo(TEXT("RightArm"), TEXT("RightArm"), 3), FTransform::Identity);
+
+				// When our modifier is destroyed here, it will rebuild the skeleton
+			}
+
+			UE_RETURN_ON_ERROR(Skeleton->GetReferenceSkeleton().GetNum() == 7, "Failed to build skeleton asset");
+
+			bool bSuccess = Binding->SetSkeleton(Skeleton);
+			UE_RETURN_ON_ERROR(bSuccess, "Failed to set skeleton on our binding");
+
+			UFactory* SecCollectionFactory = NewObject<UFactory>(GetTransientPackage(), UAbstractSkeletonSetCollectionFactory::StaticClass());
+			UE_RETURN_ON_ERROR(SecCollectionFactory != nullptr, "Failed to create factory");
+
+			UAbstractSkeletonSetCollection* SetCollection = CastChecked<UAbstractSkeletonSetCollection>(SecCollectionFactory->FactoryCreateNew(UAbstractSkeletonSetCollection::StaticClass(), GetTransientPackage(), TEXT("TestSetCollectionAsset"), RF_Transient, nullptr, nullptr, NAME_None));
+			UE_RETURN_ON_ERROR(SetCollection != nullptr, "Failed to create set collection asset");
+
+			SetCollection->AddSet(TEXT("FullBody"), NAME_None);
+			SetCollection->AddSet(TEXT("UpperBody"), TEXT("FullBody"));
+
+			bSuccess = Binding->SetSetCollection(SetCollection);
+			UE_RETURN_ON_ERROR(bSuccess, "Failed to set set collection on our binding");
+
+			Binding->AddBoneToSet(TEXT("Root"), TEXT("FullBody"));
+			Binding->AddBoneToSet(TEXT("Pelvis"), TEXT("FullBody"));
+			Binding->AddBoneToSet(TEXT("Spine_01"), TEXT("FullBody"));
+			Binding->AddBoneToSet(TEXT("Spine_02"), TEXT("FullBody"));
+			Binding->AddAttributeToSet(FAnimationAttributeIdentifier(TEXT("AttrFloat0"), INDEX_NONE, NAME_None, FFloatAnimationAttribute::StaticStruct()), TEXT("FullBody"));
+
+			Binding->AddBoneToSet(TEXT("Head"), TEXT("UpperBody"));
+			Binding->AddBoneToSet(TEXT("LeftArm"), TEXT("UpperBody"));
+			Binding->AddBoneToSet(TEXT("RightArm"), TEXT("UpperBody"));
+			Binding->AddAttributeToSet(FAnimationAttributeIdentifier(TEXT("AttrFloat1"), INDEX_NONE, NAME_None, FFloatAnimationAttribute::StaticStruct()), TEXT("UpperBody"));
+
+			BindingData = MakeAttributeBindingData(*Binding);
+			UE_RETURN_ON_ERROR(BindingData.IsValid(), TEXT("Attribute binding data should be valid"));
+		}
+
+		{
+			FAttributeNamedSetPtr FullBodySet = BindingData->FindNamedSet(TEXT("FullBody"));
+			UE_RETURN_ON_ERROR(FullBodySet.IsValid(), TEXT("FullBody named set should be valid"));
+
+			FAttributeCollectionHeap Collection(FullBodySet);
+
+			AddErrorIfFalse(Collection.IsEmpty(), TEXT("Attribute collection should be empty"));
+			AddErrorIfFalse(Collection.GetNamedSet() == FullBodySet, TEXT("Unexpected named set"));
+
+			FAttributeMappingKey FloatToFloatKey = FAttributeMappingKey::MakeFromTo<FFloatAnimationAttribute, FFloatAnimationAttribute>();
+			TAttributeSetMap<FFloatAnimationAttribute>* FloatToFloatMap = Collection.AddSetMap<FFloatAnimationAttribute>(FloatToFloatKey);
+			UE_RETURN_ON_ERROR(FloatToFloatMap != nullptr, TEXT("Float->Float map should be valid"));
+
+			FAttributeMappingKey FloatToIntegerKey = FAttributeMappingKey::MakeFromTo<FFloatAnimationAttribute, FIntegerAnimationAttribute>();
+			TAttributeSetMap<FIntegerAnimationAttribute>* FloatToIntegerMap = MakeAttributeSetMap<FIntegerAnimationAttribute>(FAttributeSetMap::FConstructArgs(FullBodySet->FindTypedSet<FFloatAnimationAttribute>(), FIntegerAnimationAttribute::StaticStruct(), Collection.GetAllocator()));
+			UE_RETURN_ON_ERROR(FloatToIntegerMap != nullptr, TEXT("Float->Integer map should be valid"));
+			bool bAdded = Collection.AddSetMap(FloatToIntegerMap);
+			UE_RETURN_ON_ERROR(bAdded, TEXT("Float->Integer map should be added"));
+
+			// Can't add twice
+			bAdded = Collection.AddSetMap(FloatToIntegerMap);
+			AddErrorIfFalse(!bAdded, TEXT("Float->Integer map should not be added"));
+
+			FAttributeMappingKey FloatToStringKey = FAttributeMappingKey::MakeFromTo<FFloatAnimationAttribute, FStringAnimationAttribute>();
+			TAttributeSetMap<FStringAnimationAttribute>* FloatToStringMap = Collection.AddSetMap<FStringAnimationAttribute>(FloatToStringKey);
+			UE_RETURN_ON_ERROR(FloatToStringMap != nullptr, TEXT("Float->String map should be valid"));
+
+			FAttributeMappingKey TransformToTransformKey = FAttributeMappingKey::MakeFromTo<FBoneTransformAnimationAttribute, FBoneTransformAnimationAttribute>();
+			TAttributeSetMap<FBoneTransformAnimationAttribute>* TransformToTransformMap = Collection.AddSetMap<FBoneTransformAnimationAttribute>(TransformToTransformKey);
+			UE_RETURN_ON_ERROR(TransformToTransformMap != nullptr, TEXT("Transform->Transform map should be valid"));
+
+			FAttributeMappingKey TransformToFloatKey = FAttributeMappingKey::MakeFromTo<FBoneTransformAnimationAttribute, FFloatAnimationAttribute>();
+			TAttributeSetMap<FFloatAnimationAttribute>* TransformToFloatMap = Cast<FFloatAnimationAttribute>(Collection.AddSetMap(TransformToFloatKey));
+			UE_RETURN_ON_ERROR(TransformToFloatMap != nullptr, TEXT("Transform->Bone map should be valid"));
+
+			TAttributeMap<FFloatAnimationAttribute>* FloatMap = Collection.AddMap<FFloatAnimationAttribute>();
+			UE_RETURN_ON_ERROR(FloatMap != nullptr, TEXT("Float map should be valid"));
+
+			TAttributeMap<FBoneTransformAnimationAttribute>* BoneMap = Cast<FBoneTransformAnimationAttribute>(Collection.AddMap(FBoneTransformAnimationAttribute::StaticStruct()));
+			UE_RETURN_ON_ERROR(BoneMap != nullptr, TEXT("Bone map should be valid"));
+
+			TAttributeMap<FIntegerAnimationAttribute>* IntMap = MakeAttributeMap<FIntegerAnimationAttribute>(Collection.GetAllocator());
+			UE_RETURN_ON_ERROR(IntMap != nullptr, TEXT("Integer map should be valid"));
+			bAdded = Collection.AddMap(IntMap);
+			UE_RETURN_ON_ERROR(bAdded, TEXT("Integer map should be added"));
+
+			// Can't add twice
+			bAdded = Collection.AddMap(IntMap);
+			AddErrorIfFalse(!bAdded, TEXT("Integer map should not be added"));
+
+			AddErrorIfFalse(!Collection.IsEmpty(), TEXT("Attribute collection should not be empty"));
+			AddErrorIfFalse(Collection.NumSetMaps() == 5, TEXT("Unexpected set map count"));
+			AddErrorIfFalse(Collection.NumMaps() == 3, TEXT("Unexpected map count"));
+
+			AddErrorIfFalse(FloatToFloatMap == Collection.FindSetMap<FFloatAnimationAttribute>(FloatToFloatKey), TEXT("Unexpected result from FindSetMap"));
+			AddErrorIfFalse(FloatToIntegerMap == Collection.FindSetMap<FIntegerAnimationAttribute>(FloatToIntegerKey), TEXT("Unexpected result from FindSetMap"));
+			AddErrorIfFalse(FloatToStringMap == Collection.FindSetMap<FStringAnimationAttribute>(FloatToStringKey), TEXT("Unexpected result from FindSetMap"));
+			AddErrorIfFalse(TransformToTransformMap == Collection.FindSetMap<FBoneTransformAnimationAttribute>(TransformToTransformKey), TEXT("Unexpected result from FindSetMap"));
+			AddErrorIfFalse(TransformToFloatMap == Collection.FindSetMap(TransformToFloatKey), TEXT("Unexpected result from FindSetMap"));
+			AddErrorIfFalse(FloatMap == Collection.FindMap<FFloatAnimationAttribute>(), TEXT("Unexpected result from FindMap"));
+			AddErrorIfFalse(BoneMap == Collection.FindMap(FBoneTransformAnimationAttribute::StaticStruct()), TEXT("Unexpected result from FindMap"));
+
+			AddErrorIfFalse(Collection.RemoveSetMap(FloatToFloatMap), TEXT("Failed to remove set map"));
+			AddErrorIfFalse(Collection.RemoveSetMap(FloatToIntegerKey), TEXT("Failed to remove set map"));
+			AddErrorIfFalse(Collection.RemoveSetMap(FloatToStringMap), TEXT("Failed to remove set map"));
+			AddErrorIfFalse(Collection.RemoveSetMap(TransformToTransformKey), TEXT("Failed to remove set map"));
+			AddErrorIfFalse(Collection.RemoveSetMap(TransformToFloatKey), TEXT("Failed to remove set map"));
+			AddErrorIfFalse(Collection.RemoveMap(FloatMap), TEXT("Failed to remove map"));
+			AddErrorIfFalse(Collection.RemoveMap(FBoneTransformAnimationAttribute::StaticStruct()), TEXT("Failed to remove map"));
+			AddErrorIfFalse(Collection.RemoveMap(IntMap), TEXT("Failed to remove map"));
+
+			AddErrorIfFalse(nullptr == Collection.FindSetMap<FFloatAnimationAttribute>(FloatToFloatKey), TEXT("Unexpected result from FindSetMap"));
+			AddErrorIfFalse(nullptr == Collection.FindSetMap<FIntegerAnimationAttribute>(FloatToIntegerKey), TEXT("Unexpected result from FindSetMap"));
+			AddErrorIfFalse(nullptr == Collection.FindSetMap<FStringAnimationAttribute>(FloatToStringKey), TEXT("Unexpected result from FindSetMap"));
+			AddErrorIfFalse(nullptr == Collection.FindSetMap<FBoneTransformAnimationAttribute>(TransformToTransformKey), TEXT("Unexpected result from FindSetMap"));
+			AddErrorIfFalse(nullptr == Collection.FindSetMap(TransformToFloatKey), TEXT("Unexpected result from FindSetMap"));
+			AddErrorIfFalse(nullptr == Collection.FindMap<FFloatAnimationAttribute>(), TEXT("Unexpected result from FindMap"));
+			AddErrorIfFalse(nullptr == Collection.FindMap(FBoneTransformAnimationAttribute::StaticStruct()), TEXT("Unexpected result from FindMap"));
+
+			AddErrorIfFalse(Collection.IsEmpty(), TEXT("Attribute collection should be empty"));
+			AddErrorIfFalse(Collection.NumSetMaps() == 0, TEXT("Unexpected set map count"));
+			AddErrorIfFalse(Collection.NumMaps() == 0, TEXT("Unexpected map count"));
+
+			FloatToFloatMap = Collection.FindOrAddSetMap<FFloatAnimationAttribute>(FloatToFloatKey);
+			UE_RETURN_ON_ERROR(FloatToFloatMap != nullptr, TEXT("Float->Float map should be valid"));
+
+			FloatToIntegerMap = Collection.FindOrAddSetMap<FIntegerAnimationAttribute>(FloatToIntegerKey);
+			UE_RETURN_ON_ERROR(FloatToIntegerMap != nullptr, TEXT("Float->Integer map should be valid"));
+
+			FloatToStringMap = Collection.FindOrAddSetMap<FStringAnimationAttribute>(FloatToStringKey);
+			UE_RETURN_ON_ERROR(FloatToStringMap != nullptr, TEXT("Float->String map should be valid"));
+
+			TransformToTransformMap = Collection.FindOrAddSetMap<FBoneTransformAnimationAttribute>(TransformToTransformKey);
+			UE_RETURN_ON_ERROR(TransformToTransformMap != nullptr, TEXT("Transform->Transform map should be valid"));
+
+			TransformToFloatMap = Cast<FFloatAnimationAttribute>(Collection.FindOrAddSetMap(TransformToFloatKey));
+			UE_RETURN_ON_ERROR(TransformToFloatMap != nullptr, TEXT("Transform->Bone map should be valid"));
+
+			FloatMap = Collection.FindOrAddMap<FFloatAnimationAttribute>();
+			UE_RETURN_ON_ERROR(FloatMap != nullptr, TEXT("Float map should be valid"));
+
+			BoneMap = Cast<FBoneTransformAnimationAttribute>(Collection.FindOrAddMap(FBoneTransformAnimationAttribute::StaticStruct()));
+			UE_RETURN_ON_ERROR(BoneMap != nullptr, TEXT("Bone map should be valid"));
+
+			AddErrorIfFalse(!Collection.IsEmpty(), TEXT("Attribute collection should not be empty"));
+			AddErrorIfFalse(Collection.NumSetMaps() == 5, TEXT("Unexpected set map count"));
+			AddErrorIfFalse(Collection.NumMaps() == 2, TEXT("Unexpected map count"));
+
+			AddErrorIfFalse(FloatToFloatMap == Collection.FindOrAddSetMap<FFloatAnimationAttribute>(FloatToFloatKey), TEXT("Unexpected result from FindOrAddSetMap"));
+			AddErrorIfFalse(FloatToIntegerMap == Collection.FindOrAddSetMap<FIntegerAnimationAttribute>(FloatToIntegerKey), TEXT("Unexpected result from FindOrAddSetMap"));
+			AddErrorIfFalse(FloatToStringMap == Collection.FindOrAddSetMap<FStringAnimationAttribute>(FloatToStringKey), TEXT("Unexpected result from FindOrAddSetMap"));
+			AddErrorIfFalse(TransformToTransformMap == Collection.FindOrAddSetMap<FBoneTransformAnimationAttribute>(TransformToTransformKey), TEXT("Unexpected result from FindOrAddSetMap"));
+			AddErrorIfFalse(TransformToFloatMap == Collection.FindOrAddSetMap(TransformToFloatKey), TEXT("Unexpected result from FindOrAddSetMap"));
+			AddErrorIfFalse(FloatMap == Collection.FindOrAddMap<FFloatAnimationAttribute>(), TEXT("Unexpected result from FindOrAddMap"));
+			AddErrorIfFalse(BoneMap == Collection.FindOrAddMap(FBoneTransformAnimationAttribute::StaticStruct()), TEXT("Unexpected result from FindOrAddMap"));
+
+			AddErrorIfFalse(!Collection.IsEmpty(), TEXT("Attribute collection should not be empty"));
+			AddErrorIfFalse(Collection.NumSetMaps() == 5, TEXT("Unexpected set map count"));
+			AddErrorIfFalse(Collection.NumMaps() == 2, TEXT("Unexpected map count"));
+
+			// Append
+			{
+				FAttributeCollectionHeap CollectionCopy(FullBodySet);
+
+				for (auto It = Collection.CreateConstSetMapIterator(); It; ++It)
+				{
+					FAttributeSetMap* SetMapCopy = It.GetMap()->Duplicate(CollectionCopy.GetAllocator());
+					const bool bAppended = CollectionCopy.AppendSetMap(SetMapCopy);
+					AddErrorIfFalse(bAppended, TEXT("Set map should be appended"));
+				}
+
+				for (auto It = Collection.CreateConstMapIterator(); It; ++It)
+				{
+					FAttributeMap* MapCopy = It.GetMap()->Duplicate(CollectionCopy.GetAllocator());
+					const bool bAppended = CollectionCopy.AppendMap(MapCopy);
+					AddErrorIfFalse(bAppended, TEXT("Map should be appended"));
+				}
+
+				AddErrorIfFalse(CollectionCopy.IsEmpty() == Collection.IsEmpty(), TEXT("Attribute collection should not be empty"));
+				AddErrorIfFalse(CollectionCopy.NumSetMaps() == Collection.NumSetMaps(), TEXT("Unexpected set map count"));
+				AddErrorIfFalse(CollectionCopy.NumMaps() == Collection.NumMaps(), TEXT("Unexpected map count"));
+
+				TAttributeSetMap<FFloatAnimationAttribute>* FloatToFloatMapCopy = CollectionCopy.FindSetMap<FFloatAnimationAttribute>(FloatToFloatKey);
+				UE_RETURN_ON_ERROR(FloatToFloatMapCopy != nullptr, TEXT("Float->Float map should be valid"));
+
+				TAttributeSetMap<FIntegerAnimationAttribute>* FloatToIntegerMapCopy = CollectionCopy.FindSetMap<FIntegerAnimationAttribute>(FloatToIntegerKey);
+				UE_RETURN_ON_ERROR(FloatToIntegerMapCopy != nullptr, TEXT("Float->Integer map should be valid"));
+
+				TAttributeSetMap<FStringAnimationAttribute>* FloatToStringMapCopy = CollectionCopy.FindSetMap<FStringAnimationAttribute>(FloatToStringKey);
+				UE_RETURN_ON_ERROR(FloatToStringMapCopy != nullptr, TEXT("Float->String map should be valid"));
+
+				TAttributeSetMap<FBoneTransformAnimationAttribute>* TransformToTransformMapCopy = CollectionCopy.FindSetMap<FBoneTransformAnimationAttribute>(TransformToTransformKey);
+				UE_RETURN_ON_ERROR(TransformToTransformMapCopy != nullptr, TEXT("Transform->Transform map should be valid"));
+
+				TAttributeSetMap<FFloatAnimationAttribute>* TransformToFloatMapCopy = Cast<FFloatAnimationAttribute>(CollectionCopy.FindSetMap(TransformToFloatKey));
+				UE_RETURN_ON_ERROR(TransformToFloatMapCopy != nullptr, TEXT("Transform->Bone map should be valid"));
+
+				TAttributeMap<FFloatAnimationAttribute>* FloatMapCopy = CollectionCopy.FindMap<FFloatAnimationAttribute>();
+				UE_RETURN_ON_ERROR(FloatMapCopy != nullptr, TEXT("Float map should be valid"));
+
+				TAttributeMap<FBoneTransformAnimationAttribute>* BoneMapCopy = Cast<FBoneTransformAnimationAttribute>(CollectionCopy.FindMap(FBoneTransformAnimationAttribute::StaticStruct()));
+				UE_RETURN_ON_ERROR(BoneMapCopy != nullptr, TEXT("Bone map should be valid"));
+
+				AddErrorIfFalse(FloatToFloatMap != FloatToFloatMapCopy, TEXT("Set map instances should be duplicated"));
+				AddErrorIfFalse(FloatToIntegerMap != FloatToIntegerMapCopy, TEXT("Set map instances should be duplicated"));
+				AddErrorIfFalse(FloatToStringMap != FloatToStringMapCopy, TEXT("Set map instances should be duplicated"));
+				AddErrorIfFalse(TransformToTransformMap != TransformToTransformMapCopy, TEXT("Set map instances should be duplicated"));
+				AddErrorIfFalse(TransformToFloatMap != TransformToFloatMapCopy, TEXT("Set map instances should be duplicated"));
+				AddErrorIfFalse(FloatMap != FloatMapCopy, TEXT("Name/value map instances should be duplicated"));
+				AddErrorIfFalse(BoneMap != BoneMapCopy, TEXT("Name/value map instances should be duplicated"));
+
+				AddErrorIfFalse(FloatToFloatMap->GetTypedSet() == FloatToFloatMapCopy->GetTypedSet(), TEXT("Typed set mismatched"));
+				AddErrorIfFalse(FloatToIntegerMap->GetTypedSet() == FloatToIntegerMapCopy->GetTypedSet(), TEXT("Typed set mismatched"));
+				AddErrorIfFalse(FloatToStringMap->GetTypedSet() == FloatToStringMapCopy->GetTypedSet(), TEXT("Typed set mismatched"));
+				AddErrorIfFalse(TransformToTransformMap->GetTypedSet() == TransformToTransformMapCopy->GetTypedSet(), TEXT("Typed set mismatched"));
+				AddErrorIfFalse(TransformToFloatMap->GetTypedSet() == TransformToFloatMapCopy->GetTypedSet(), TEXT("Typed set mismatched"));
+			}
+
+			// Copy constructor
+			{
+				FAttributeCollectionHeap CollectionCopy(Collection);
+
+				AddErrorIfFalse(CollectionCopy.IsEmpty() == Collection.IsEmpty(), TEXT("Attribute collection should not be empty"));
+				AddErrorIfFalse(CollectionCopy.NumSetMaps() == Collection.NumSetMaps(), TEXT("Unexpected set map count"));
+				AddErrorIfFalse(CollectionCopy.NumMaps() == Collection.NumMaps(), TEXT("Unexpected map count"));
+
+				TAttributeSetMap<FFloatAnimationAttribute>* FloatToFloatMapCopy = CollectionCopy.FindSetMap<FFloatAnimationAttribute>(FloatToFloatKey);
+				UE_RETURN_ON_ERROR(FloatToFloatMapCopy != nullptr, TEXT("Float->Float map should be valid"));
+
+				TAttributeSetMap<FIntegerAnimationAttribute>* FloatToIntegerMapCopy = CollectionCopy.FindSetMap<FIntegerAnimationAttribute>(FloatToIntegerKey);
+				UE_RETURN_ON_ERROR(FloatToIntegerMapCopy != nullptr, TEXT("Float->Integer map should be valid"));
+
+				TAttributeSetMap<FStringAnimationAttribute>* FloatToStringMapCopy = CollectionCopy.FindSetMap<FStringAnimationAttribute>(FloatToStringKey);
+				UE_RETURN_ON_ERROR(FloatToStringMapCopy != nullptr, TEXT("Float->String map should be valid"));
+
+				TAttributeSetMap<FBoneTransformAnimationAttribute>* TransformToTransformMapCopy = CollectionCopy.FindSetMap<FBoneTransformAnimationAttribute>(TransformToTransformKey);
+				UE_RETURN_ON_ERROR(TransformToTransformMapCopy != nullptr, TEXT("Transform->Transform map should be valid"));
+
+				TAttributeSetMap<FFloatAnimationAttribute>* TransformToFloatMapCopy = Cast<FFloatAnimationAttribute>(CollectionCopy.FindSetMap(TransformToFloatKey));
+				UE_RETURN_ON_ERROR(TransformToFloatMapCopy != nullptr, TEXT("Transform->Bone map should be valid"));
+
+				TAttributeMap<FFloatAnimationAttribute>* FloatMapCopy = CollectionCopy.FindMap<FFloatAnimationAttribute>();
+				UE_RETURN_ON_ERROR(FloatMapCopy != nullptr, TEXT("Float map should be valid"));
+
+				TAttributeMap<FBoneTransformAnimationAttribute>* BoneMapCopy = Cast<FBoneTransformAnimationAttribute>(CollectionCopy.FindMap(FBoneTransformAnimationAttribute::StaticStruct()));
+				UE_RETURN_ON_ERROR(BoneMapCopy != nullptr, TEXT("Bone map should be valid"));
+
+				AddErrorIfFalse(FloatToFloatMap != FloatToFloatMapCopy, TEXT("Set map instances should be duplicated"));
+				AddErrorIfFalse(FloatToIntegerMap != FloatToIntegerMapCopy, TEXT("Set map instances should be duplicated"));
+				AddErrorIfFalse(FloatToStringMap != FloatToStringMapCopy, TEXT("Set map instances should be duplicated"));
+				AddErrorIfFalse(TransformToTransformMap != TransformToTransformMapCopy, TEXT("Set map instances should be duplicated"));
+				AddErrorIfFalse(TransformToFloatMap != TransformToFloatMapCopy, TEXT("Set map instances should be duplicated"));
+				AddErrorIfFalse(FloatMap != FloatMapCopy, TEXT("Name/value map instances should be duplicated"));
+				AddErrorIfFalse(BoneMap != BoneMapCopy, TEXT("Name/value map instances should be duplicated"));
+
+				AddErrorIfFalse(FloatToFloatMap->GetTypedSet() == FloatToFloatMapCopy->GetTypedSet(), TEXT("Typed set mismatched"));
+				AddErrorIfFalse(FloatToIntegerMap->GetTypedSet() == FloatToIntegerMapCopy->GetTypedSet(), TEXT("Typed set mismatched"));
+				AddErrorIfFalse(FloatToStringMap->GetTypedSet() == FloatToStringMapCopy->GetTypedSet(), TEXT("Typed set mismatched"));
+				AddErrorIfFalse(TransformToTransformMap->GetTypedSet() == TransformToTransformMapCopy->GetTypedSet(), TEXT("Typed set mismatched"));
+				AddErrorIfFalse(TransformToFloatMap->GetTypedSet() == TransformToFloatMapCopy->GetTypedSet(), TEXT("Typed set mismatched"));
+
+				// Move constructor
+				{
+					FAttributeCollectionHeap CollectionMoved(MoveTemp(CollectionCopy));
+
+					AddErrorIfFalse(CollectionCopy.IsEmpty(), TEXT("Attribute collection should be empty"));
+					AddErrorIfFalse(!CollectionMoved.IsEmpty(), TEXT("Attribute collection should not be empty"));
+					AddErrorIfFalse(CollectionMoved.NumSetMaps() == Collection.NumSetMaps(), TEXT("Unexpected set map count"));
+					AddErrorIfFalse(CollectionMoved.NumMaps() == Collection.NumMaps(), TEXT("Unexpected map count"));
+
+					TAttributeSetMap<FFloatAnimationAttribute>* FloatToFloatMapMoved = CollectionMoved.FindSetMap<FFloatAnimationAttribute>(FloatToFloatKey);
+					UE_RETURN_ON_ERROR(FloatToFloatMapMoved != nullptr, TEXT("Float->Float map should be valid"));
+
+					TAttributeSetMap<FIntegerAnimationAttribute>* FloatToIntegerMapMoved = CollectionMoved.FindSetMap<FIntegerAnimationAttribute>(FloatToIntegerKey);
+					UE_RETURN_ON_ERROR(FloatToIntegerMapMoved != nullptr, TEXT("Float->Integer map should be valid"));
+
+					TAttributeSetMap<FStringAnimationAttribute>* FloatToStringMapMoved = CollectionMoved.FindSetMap<FStringAnimationAttribute>(FloatToStringKey);
+					UE_RETURN_ON_ERROR(FloatToStringMapMoved != nullptr, TEXT("Float->String map should be valid"));
+
+					TAttributeSetMap<FBoneTransformAnimationAttribute>* TransformToTransformMapMoved = CollectionMoved.FindSetMap<FBoneTransformAnimationAttribute>(TransformToTransformKey);
+					UE_RETURN_ON_ERROR(TransformToTransformMapMoved != nullptr, TEXT("Transform->Transform map should be valid"));
+
+					TAttributeSetMap<FFloatAnimationAttribute>* TransformToFloatMapMoved = Cast<FFloatAnimationAttribute>(CollectionMoved.FindSetMap(TransformToFloatKey));
+					UE_RETURN_ON_ERROR(TransformToFloatMapMoved != nullptr, TEXT("Transform->Bone map should be valid"));
+
+					TAttributeMap<FFloatAnimationAttribute>* FloatMapMoved = CollectionMoved.FindMap<FFloatAnimationAttribute>();
+					UE_RETURN_ON_ERROR(FloatMapMoved != nullptr, TEXT("Float map should be valid"));
+
+					TAttributeMap<FBoneTransformAnimationAttribute>* BoneMapMoved = Cast<FBoneTransformAnimationAttribute>(CollectionMoved.FindMap(FBoneTransformAnimationAttribute::StaticStruct()));
+					UE_RETURN_ON_ERROR(BoneMapMoved != nullptr, TEXT("Bone map should be valid"));
+
+					AddErrorIfFalse(FloatToFloatMapMoved == FloatToFloatMapCopy, TEXT("Set map instances should be moved"));
+					AddErrorIfFalse(FloatToIntegerMapMoved == FloatToIntegerMapCopy, TEXT("Set map instances should be moved"));
+					AddErrorIfFalse(FloatToStringMapMoved == FloatToStringMapCopy, TEXT("Set map instances should be moved"));
+					AddErrorIfFalse(TransformToTransformMapMoved == TransformToTransformMapCopy, TEXT("Set map instances should be moved"));
+					AddErrorIfFalse(TransformToFloatMapMoved == TransformToFloatMapCopy, TEXT("Set map instances should be moved"));
+					AddErrorIfFalse(FloatMapMoved == FloatMapCopy, TEXT("Set map instances should be moved"));
+					AddErrorIfFalse(BoneMapMoved == BoneMapCopy, TEXT("Set map instances should be moved"));
+				}
+			}
+
+			// Copy assignment
+			{
+				FAttributeCollectionHeap CollectionCopy;
+				CollectionCopy = Collection;
+
+				AddErrorIfFalse(CollectionCopy.IsEmpty() == Collection.IsEmpty(), TEXT("Attribute collection should not be empty"));
+				AddErrorIfFalse(CollectionCopy.NumSetMaps() == Collection.NumSetMaps(), TEXT("Unexpected set map count"));
+				AddErrorIfFalse(CollectionCopy.NumMaps() == Collection.NumMaps(), TEXT("Unexpected map count"));
+
+				TAttributeSetMap<FFloatAnimationAttribute>* FloatToFloatMapCopy = CollectionCopy.FindSetMap<FFloatAnimationAttribute>(FloatToFloatKey);
+				UE_RETURN_ON_ERROR(FloatToFloatMapCopy != nullptr, TEXT("Float->Float map should be valid"));
+
+				TAttributeSetMap<FIntegerAnimationAttribute>* FloatToIntegerMapCopy = CollectionCopy.FindSetMap<FIntegerAnimationAttribute>(FloatToIntegerKey);
+				UE_RETURN_ON_ERROR(FloatToIntegerMapCopy != nullptr, TEXT("Float->Integer map should be valid"));
+
+				TAttributeSetMap<FStringAnimationAttribute>* FloatToStringMapCopy = CollectionCopy.FindSetMap<FStringAnimationAttribute>(FloatToStringKey);
+				UE_RETURN_ON_ERROR(FloatToStringMapCopy != nullptr, TEXT("Float->String map should be valid"));
+
+				TAttributeSetMap<FBoneTransformAnimationAttribute>* TransformToTransformMapCopy = CollectionCopy.FindSetMap<FBoneTransformAnimationAttribute>(TransformToTransformKey);
+				UE_RETURN_ON_ERROR(TransformToTransformMapCopy != nullptr, TEXT("Transform->Transform map should be valid"));
+
+				TAttributeSetMap<FFloatAnimationAttribute>* TransformToFloatMapCopy = Cast<FFloatAnimationAttribute>(CollectionCopy.FindSetMap(TransformToFloatKey));
+				UE_RETURN_ON_ERROR(TransformToFloatMapCopy != nullptr, TEXT("Transform->Bone map should be valid"));
+
+				TAttributeMap<FFloatAnimationAttribute>* FloatMapCopy = CollectionCopy.FindMap<FFloatAnimationAttribute>();
+				UE_RETURN_ON_ERROR(FloatMapCopy != nullptr, TEXT("Float map should be valid"));
+
+				TAttributeMap<FBoneTransformAnimationAttribute>* BoneMapCopy = Cast<FBoneTransformAnimationAttribute>(CollectionCopy.FindMap(FBoneTransformAnimationAttribute::StaticStruct()));
+				UE_RETURN_ON_ERROR(BoneMapCopy != nullptr, TEXT("Bone map should be valid"));
+
+				AddErrorIfFalse(FloatToFloatMap != FloatToFloatMapCopy, TEXT("Set map instances should be duplicated"));
+				AddErrorIfFalse(FloatToIntegerMap != FloatToIntegerMapCopy, TEXT("Set map instances should be duplicated"));
+				AddErrorIfFalse(FloatToStringMap != FloatToStringMapCopy, TEXT("Set map instances should be duplicated"));
+				AddErrorIfFalse(TransformToTransformMap != TransformToTransformMapCopy, TEXT("Set map instances should be duplicated"));
+				AddErrorIfFalse(TransformToFloatMap != TransformToFloatMapCopy, TEXT("Set map instances should be duplicated"));
+				AddErrorIfFalse(FloatMap != FloatMapCopy, TEXT("Name/value map instances should be duplicated"));
+				AddErrorIfFalse(BoneMap != BoneMapCopy, TEXT("Name/value map instances should be duplicated"));
+
+				AddErrorIfFalse(FloatToFloatMap->GetTypedSet() == FloatToFloatMapCopy->GetTypedSet(), TEXT("Typed set mismatched"));
+				AddErrorIfFalse(FloatToIntegerMap->GetTypedSet() == FloatToIntegerMapCopy->GetTypedSet(), TEXT("Typed set mismatched"));
+				AddErrorIfFalse(FloatToStringMap->GetTypedSet() == FloatToStringMapCopy->GetTypedSet(), TEXT("Typed set mismatched"));
+				AddErrorIfFalse(TransformToTransformMap->GetTypedSet() == TransformToTransformMapCopy->GetTypedSet(), TEXT("Typed set mismatched"));
+				AddErrorIfFalse(TransformToFloatMap->GetTypedSet() == TransformToFloatMapCopy->GetTypedSet(), TEXT("Typed set mismatched"));
+
+				// Move assignment
+				{
+					FAttributeCollectionHeap CollectionMoved;
+					CollectionMoved = MoveTemp(CollectionCopy);
+
+					AddErrorIfFalse(CollectionCopy.IsEmpty(), TEXT("Attribute collection should be empty"));
+					AddErrorIfFalse(!CollectionMoved.IsEmpty(), TEXT("Attribute collection should not be empty"));
+					AddErrorIfFalse(CollectionMoved.NumSetMaps() == Collection.NumSetMaps(), TEXT("Unexpected set map count"));
+					AddErrorIfFalse(CollectionMoved.NumMaps() == Collection.NumMaps(), TEXT("Unexpected map count"));
+
+					TAttributeSetMap<FFloatAnimationAttribute>* FloatToFloatMapMoved = CollectionMoved.FindSetMap<FFloatAnimationAttribute>(FloatToFloatKey);
+					UE_RETURN_ON_ERROR(FloatToFloatMapMoved != nullptr, TEXT("Float->Float map should be valid"));
+
+					TAttributeSetMap<FIntegerAnimationAttribute>* FloatToIntegerMapMoved = CollectionMoved.FindSetMap<FIntegerAnimationAttribute>(FloatToIntegerKey);
+					UE_RETURN_ON_ERROR(FloatToIntegerMapMoved != nullptr, TEXT("Float->Integer map should be valid"));
+
+					TAttributeSetMap<FStringAnimationAttribute>* FloatToStringMapMoved = CollectionMoved.FindSetMap<FStringAnimationAttribute>(FloatToStringKey);
+					UE_RETURN_ON_ERROR(FloatToStringMapMoved != nullptr, TEXT("Float->String map should be valid"));
+
+					TAttributeSetMap<FBoneTransformAnimationAttribute>* TransformToTransformMapMoved = CollectionMoved.FindSetMap<FBoneTransformAnimationAttribute>(TransformToTransformKey);
+					UE_RETURN_ON_ERROR(TransformToTransformMapMoved != nullptr, TEXT("Transform->Transform map should be valid"));
+
+					TAttributeSetMap<FFloatAnimationAttribute>* TransformToFloatMapMoved = Cast<FFloatAnimationAttribute>(CollectionMoved.FindSetMap(TransformToFloatKey));
+					UE_RETURN_ON_ERROR(TransformToFloatMapMoved != nullptr, TEXT("Transform->Bone map should be valid"));
+
+					TAttributeMap<FFloatAnimationAttribute>* FloatMapMoved = CollectionMoved.FindMap<FFloatAnimationAttribute>();
+					UE_RETURN_ON_ERROR(FloatMapMoved != nullptr, TEXT("Float map should be valid"));
+
+					TAttributeMap<FBoneTransformAnimationAttribute>* BoneMapMoved = Cast<FBoneTransformAnimationAttribute>(CollectionMoved.FindMap(FBoneTransformAnimationAttribute::StaticStruct()));
+					UE_RETURN_ON_ERROR(BoneMapMoved != nullptr, TEXT("Bone map should be valid"));
+
+					AddErrorIfFalse(FloatToFloatMapMoved == FloatToFloatMapCopy, TEXT("Set map instances should be moved"));
+					AddErrorIfFalse(FloatToIntegerMapMoved == FloatToIntegerMapCopy, TEXT("Set map instances should be moved"));
+					AddErrorIfFalse(FloatToStringMapMoved == FloatToStringMapCopy, TEXT("Set map instances should be moved"));
+					AddErrorIfFalse(TransformToTransformMapMoved == TransformToTransformMapCopy, TEXT("Set map instances should be moved"));
+					AddErrorIfFalse(TransformToFloatMapMoved == TransformToFloatMapCopy, TEXT("Set map instances should be moved"));
+					AddErrorIfFalse(FloatMapMoved == FloatMapCopy, TEXT("Name/value map instances should be moved"));
+					AddErrorIfFalse(BoneMapMoved == BoneMapCopy, TEXT("Name/value map instances should be moved"));
+				}
+			}
+
+			// Iterators (set map)
+			{
+				const auto GetSetMapKeyProjection = [](const typename FAttributeCollection::FSetMapIterator& It) { return It.GetKey(); };
+				const auto GetSetMapKeyProjectionConst = [](const typename FAttributeCollection::FConstSetMapIterator& It) { return It.GetKey(); };
+				const auto GetSetMapAttributeTypeProjection = [](const typename FAttributeCollection::FSetMapIterator& It) { return It.GetAttributeType(); };
+				const auto GetSetMapAttributeTypeProjectionConst = [](const typename FAttributeCollection::FConstSetMapIterator& It) { return It.GetAttributeType(); };
+				const auto GetSetMapProjection = [](const typename FAttributeCollection::FSetMapIterator& It) { return It.GetMap(); };
+				const auto GetSetMapProjectionConst = [](const typename FAttributeCollection::FConstSetMapIterator& It) { return It.GetMap(); };
+
+				AddErrorIfFalse(IteratorSize(Collection.CreateSetMapIterator()) == 5, TEXT("Iterator did not have expected size"));
+				AddErrorIfFalse(IteratorSize(Collection.CreateConstSetMapIterator()) == 5, TEXT("Iterator did not have expected size"));
+
+				AddErrorIfFalse(IteratorSortedByPredicate(Collection.CreateSetMapIterator(), GetSetMapKeyProjection, TLess<FAttributeMappingKey>()), TEXT("Iterator should be sorted"));
+				AddErrorIfFalse(IteratorSortedByPredicate(Collection.CreateConstSetMapIterator(), GetSetMapKeyProjectionConst, TLess<FAttributeMappingKey>()), TEXT("Iterator should be sorted"));
+
+				AddErrorIfFalse(FindWithBy(Collection.CreateSetMapIterator(), FloatToFloatMap->GetTypedSet()->GetType(), GetSetMapAttributeTypeProjection), TEXT("Iterator did not contain expected attribute type"));
+				AddErrorIfFalse(FindWithBy(Collection.CreateSetMapIterator(), FloatToIntegerMap->GetTypedSet()->GetType(), GetSetMapAttributeTypeProjection), TEXT("Iterator did not contain expected attribute type"));
+				AddErrorIfFalse(FindWithBy(Collection.CreateSetMapIterator(), FloatToStringMap->GetTypedSet()->GetType(), GetSetMapAttributeTypeProjection), TEXT("Iterator did not contain expected attribute type"));
+				AddErrorIfFalse(FindWithBy(Collection.CreateSetMapIterator(), TransformToTransformMap->GetTypedSet()->GetType(), GetSetMapAttributeTypeProjection), TEXT("Iterator did not contain expected attribute type"));
+				AddErrorIfFalse(FindWithBy(Collection.CreateSetMapIterator(), TransformToFloatMap->GetTypedSet()->GetType(), GetSetMapAttributeTypeProjection), TEXT("Iterator did not contain expected attribute type"));
+
+				AddErrorIfFalse(FindWithBy(Collection.CreateSetMapIterator(), FloatToFloatMap, GetSetMapProjection), TEXT("Iterator did not contain expected set map"));
+				AddErrorIfFalse(FindWithBy(Collection.CreateSetMapIterator(), FloatToIntegerMap, GetSetMapProjection), TEXT("Iterator did not contain expected set map"));
+				AddErrorIfFalse(FindWithBy(Collection.CreateSetMapIterator(), FloatToStringMap, GetSetMapProjection), TEXT("Iterator did not contain expected set map"));
+				AddErrorIfFalse(FindWithBy(Collection.CreateSetMapIterator(), TransformToTransformMap, GetSetMapProjection), TEXT("Iterator did not contain expected set map"));
+				AddErrorIfFalse(FindWithBy(Collection.CreateSetMapIterator(), TransformToFloatMap, GetSetMapProjection), TEXT("Iterator did not contain expected set map"));
+
+				AddErrorIfFalse(FindWithBy(Collection.CreateConstSetMapIterator(), FloatToFloatMap->GetTypedSet()->GetType(), GetSetMapAttributeTypeProjectionConst), TEXT("Iterator did not contain expected attribute type"));
+				AddErrorIfFalse(FindWithBy(Collection.CreateConstSetMapIterator(), FloatToIntegerMap->GetTypedSet()->GetType(), GetSetMapAttributeTypeProjectionConst), TEXT("Iterator did not contain expected attribute type"));
+				AddErrorIfFalse(FindWithBy(Collection.CreateConstSetMapIterator(), FloatToStringMap->GetTypedSet()->GetType(), GetSetMapAttributeTypeProjectionConst), TEXT("Iterator did not contain expected attribute type"));
+				AddErrorIfFalse(FindWithBy(Collection.CreateConstSetMapIterator(), TransformToTransformMap->GetTypedSet()->GetType(), GetSetMapAttributeTypeProjectionConst), TEXT("Iterator did not contain expected attribute type"));
+				AddErrorIfFalse(FindWithBy(Collection.CreateConstSetMapIterator(), TransformToFloatMap->GetTypedSet()->GetType(), GetSetMapAttributeTypeProjectionConst), TEXT("Iterator did not contain expected attribute type"));
+
+				AddErrorIfFalse(FindWithBy(Collection.CreateConstSetMapIterator(), FloatToFloatMap, GetSetMapProjectionConst), TEXT("Iterator did not contain expected set map"));
+				AddErrorIfFalse(FindWithBy(Collection.CreateConstSetMapIterator(), FloatToIntegerMap, GetSetMapProjectionConst), TEXT("Iterator did not contain expected set map"));
+				AddErrorIfFalse(FindWithBy(Collection.CreateConstSetMapIterator(), FloatToStringMap, GetSetMapProjectionConst), TEXT("Iterator did not contain expected set map"));
+				AddErrorIfFalse(FindWithBy(Collection.CreateConstSetMapIterator(), TransformToTransformMap, GetSetMapProjectionConst), TEXT("Iterator did not contain expected set map"));
+				AddErrorIfFalse(FindWithBy(Collection.CreateConstSetMapIterator(), TransformToFloatMap, GetSetMapProjectionConst), TEXT("Iterator did not contain expected set map"));
+			}
+
+			// Iterators (name/value map)
+			{
+				const auto GetMapValueTypeProjection = [](const typename FAttributeCollection::FMapIterator& It) { return It.GetValueType(); };
+				const auto GetMapValueTypeProjectionConst = [](const typename FAttributeCollection::FConstMapIterator& It) { return It.GetValueType(); };
+				const auto GetMapProjection = [](const typename FAttributeCollection::FMapIterator& It) { return It.GetMap(); };
+				const auto GetMapProjectionConst = [](const typename FAttributeCollection::FConstMapIterator& It) { return It.GetMap(); };
+
+				AddErrorIfFalse(IteratorSize(Collection.CreateMapIterator()) == 2, TEXT("Iterator did not have expected size"));
+				AddErrorIfFalse(IteratorSize(Collection.CreateConstMapIterator()) == 2, TEXT("Iterator did not have expected size"));
+
+				AddErrorIfFalse(IteratorSortedByPredicate(Collection.CreateMapIterator(), GetMapValueTypeProjection, TLess<UScriptStruct*>()), TEXT("Iterator should be sorted"));
+				AddErrorIfFalse(IteratorSortedByPredicate(Collection.CreateConstMapIterator(), GetMapValueTypeProjectionConst, TLess<UScriptStruct*>()), TEXT("Iterator should be sorted"));
+
+				AddErrorIfFalse(FindWithBy(Collection.CreateMapIterator(), FloatMap, GetMapProjection), TEXT("Iterator did not contain expected map"));
+				AddErrorIfFalse(FindWithBy(Collection.CreateMapIterator(), BoneMap, GetMapProjection), TEXT("Iterator did not contain expected map"));
+
+				AddErrorIfFalse(FindWithBy(Collection.CreateConstMapIterator(), FloatMap, GetMapProjectionConst), TEXT("Iterator did not contain expected map"));
+				AddErrorIfFalse(FindWithBy(Collection.CreateConstMapIterator(), BoneMap, GetMapProjectionConst), TEXT("Iterator did not contain expected map"));
+			}
+		}
+
+		return true;
+	}
+}
+
+#endif	// WITH_DEV_AUTOMATION_TESTS
