Move and rename UAF ValueRuntime related files

diff --git a/Engine/Plugins/Experimental/UAF/UAFTests/Source/UAFTestSuite/Private/UAF/ValueRuntime/IteratorUtilsTest.cpp b/Engine/Plugins/Experimental/UAF/UAFTests/Source/UAFTestSuite/Private/UAF/ValueRuntime/IteratorUtilsTest.cpp
new file mode 100644
--- /dev/null
+++ b/Engine/Plugins/Experimental/UAF/UAFTests/Source/UAFTestSuite/Private/UAF/ValueRuntime/IteratorUtilsTest.cpp
@@ -0,0 +1,1754 @@
+// Copyright Epic Games, Inc. All Rights Reserved.
+
+#include "CoreMinimal.h"
+#include "Misc/AutomationTest.h"
+
+#if WITH_DEV_AUTOMATION_TESTS
+
+#include "AnimNextTest.h"
+#include "Animation/BuiltInAttributeTypes.h"
+#include "UAF/AbstractSkeleton/SetBindingFactory.h"
+#include "UAF/AbstractSkeleton/SetCollectionFactory.h"
+#include "UAF/AbstractSkeleton/AbstractSkeletonSetBinding.h"
+#include "UAF/AbstractSkeleton/AbstractSkeletonSetCollection.h"
+#include "UAF/Attributes/AttributeBindingData.h"
+#include "UAF/Attributes/EngineAttributes.h"
+#include "UAF/ValueRuntime/AttributeCollection.h"
+#include "UAF/ValueRuntime/IteratorUtils.h"
+#include "UAF/ValueRuntime/Transformers/Transformer.h"
+#include "UAF/ValueRuntime/Transformers/Sanitize.h"
+#include "UObject/UObjectGlobals.h"
+#include "UObject/Package.h"
+
+namespace UE::UAF::Tests
+{
+	static bool CollectionContains(const FAttributeCollection& Collection, const FAttributeSetMap* SetMap)
+	{
+		if (SetMap == nullptr)
+		{
+			return false;
+		}
+
+		const FAttributeSetMap* Other = Collection.FindSetMap(SetMap->GetMappingKey());
+		return SetMap == Other;
+	}
+
+	static bool CollectionContains(const FAttributeCollection& Collection, const FAttributeMap* Map)
+	{
+		if (Map == nullptr)
+		{
+			return false;
+		}
+
+		const FAttributeMap* Other = Collection.FindMap(Map->GetValueType());
+		return Map == Other;
+	}
+
+	IMPLEMENT_SIMPLE_AUTOMATION_TEST(FIteratorUtilsInnerJoinTest, "Animation.UAF.Attributes.IteratorUtils.InnerJoin", EAutomationTestFlags::EditorContext | EAutomationTestFlags::EngineFilter)
+
+	bool FIteratorUtilsInnerJoinTest::RunTest(const FString& InParameters)
+	{
+		ON_SCOPE_EXIT{ FUtils::CleanupAfterTests(); };
+
+		FAttributeBindingDataPtr BindingData;
+		{
+			UFactory* BindingFactory = NewObject<UFactory>(GetTransientPackage(), UAbstractSkeletonSetBindingFactory::StaticClass());
+			UE_RETURN_ON_ERROR(BindingFactory != nullptr, "Failed to create factory");
+
+			UAbstractSkeletonSetBinding* Binding = CastChecked<UAbstractSkeletonSetBinding>(BindingFactory->FactoryCreateNew(UAbstractSkeletonSetBinding::StaticClass(), GetTransientPackage(), TEXT("TestBindingAsset"), RF_Transient, nullptr, nullptr, NAME_None));
+			UE_RETURN_ON_ERROR(Binding != nullptr, "Failed to create binding asset");
+
+			USkeleton* Skeleton = NewObject<USkeleton>(GetTransientPackage(), USkeleton::StaticClass(), TEXT("TestSkeletonAsset"), RF_Transient);
+			UE_RETURN_ON_ERROR(Skeleton != nullptr, "Failed to create skeleton asset");
+
+			{
+				// Not terribly clean, we cast away the 'const' to modify the skeleton
+				FReferenceSkeleton& RefSkeleton = const_cast<FReferenceSkeleton&>(Skeleton->GetReferenceSkeleton());
+				FReferenceSkeletonModifier SkeletonModifier(RefSkeleton, Skeleton);
+
+				SkeletonModifier.Add(FMeshBoneInfo(TEXT("Root"), TEXT("Root"), INDEX_NONE), FTransform::Identity);
+				SkeletonModifier.Add(FMeshBoneInfo(TEXT("Pelvis"), TEXT("Root"), 0), FTransform::Identity);
+				SkeletonModifier.Add(FMeshBoneInfo(TEXT("Spine_01"), TEXT("Spine_01"), 1), FTransform::Identity);
+				SkeletonModifier.Add(FMeshBoneInfo(TEXT("Spine_02"), TEXT("Spine_02"), 2), FTransform::Identity);
+				SkeletonModifier.Add(FMeshBoneInfo(TEXT("Head"), TEXT("Head"), 3), FTransform::Identity);
+				SkeletonModifier.Add(FMeshBoneInfo(TEXT("LeftArm"), TEXT("LeftArm"), 3), FTransform::Identity);
+				SkeletonModifier.Add(FMeshBoneInfo(TEXT("RightArm"), TEXT("RightArm"), 3), FTransform::Identity);
+
+				// When our modifier is destroyed here, it will rebuild the skeleton
+			}
+
+			UE_RETURN_ON_ERROR(Skeleton->GetReferenceSkeleton().GetNum() == 7, "Failed to build skeleton asset");
+
+			bool bSuccess = Binding->SetSkeleton(Skeleton);
+			UE_RETURN_ON_ERROR(bSuccess, "Failed to set skeleton on our binding");
+
+			UFactory* SecCollectionFactory = NewObject<UFactory>(GetTransientPackage(), UAbstractSkeletonSetCollectionFactory::StaticClass());
+			UE_RETURN_ON_ERROR(SecCollectionFactory != nullptr, "Failed to create factory");
+
+			UAbstractSkeletonSetCollection* SetCollection = CastChecked<UAbstractSkeletonSetCollection>(SecCollectionFactory->FactoryCreateNew(UAbstractSkeletonSetCollection::StaticClass(), GetTransientPackage(), TEXT("TestSetCollectionAsset"), RF_Transient, nullptr, nullptr, NAME_None));
+			UE_RETURN_ON_ERROR(SetCollection != nullptr, "Failed to create set collection asset");
+
+			SetCollection->AddSet(TEXT("FullBody"), NAME_None);
+			SetCollection->AddSet(TEXT("UpperBody"), TEXT("FullBody"));
+
+			bSuccess = Binding->SetSetCollection(SetCollection);
+			UE_RETURN_ON_ERROR(bSuccess, "Failed to set set collection on our binding");
+
+			Binding->AddBoneToSet(TEXT("Root"), TEXT("FullBody"));
+			Binding->AddBoneToSet(TEXT("Pelvis"), TEXT("FullBody"));
+			Binding->AddBoneToSet(TEXT("Spine_01"), TEXT("FullBody"));
+			Binding->AddBoneToSet(TEXT("Spine_02"), TEXT("FullBody"));
+			Binding->AddAttributeToSet(FAnimationAttributeIdentifier(TEXT("AttrFloat0"), INDEX_NONE, NAME_None, FFloatAnimationAttribute::StaticStruct()), TEXT("FullBody"));
+
+			Binding->AddBoneToSet(TEXT("Head"), TEXT("UpperBody"));
+			Binding->AddBoneToSet(TEXT("LeftArm"), TEXT("UpperBody"));
+			Binding->AddBoneToSet(TEXT("RightArm"), TEXT("UpperBody"));
+			Binding->AddAttributeToSet(FAnimationAttributeIdentifier(TEXT("AttrFloat1"), INDEX_NONE, NAME_None, FFloatAnimationAttribute::StaticStruct()), TEXT("UpperBody"));
+
+			BindingData = MakeAttributeBindingData(*Binding);
+			UE_RETURN_ON_ERROR(BindingData.IsValid(), TEXT("Attribute binding data should be valid"));
+		}
+
+		{
+			FAttributeNamedSetPtr FullBodySet = BindingData->FindNamedSet(TEXT("FullBody"));
+			UE_RETURN_ON_ERROR(FullBodySet.IsValid(), TEXT("FullBody named set should be valid"));
+
+			FAttributeCollectionHeap Collection0(FullBodySet);
+
+			AddErrorIfFalse(Collection0.IsEmpty(), TEXT("Attribute collection should be empty"));
+			AddErrorIfFalse(Collection0.GetNamedSet() == FullBodySet, TEXT("Unexpected named set"));
+
+			FAttributeMappingKey FloatToFloatKey = FAttributeMappingKey::MakeFromTo<FFloatAnimationAttribute, FFloatAnimationAttribute>();
+			TAttributeSetMap<FFloatAnimationAttribute>* FloatToFloatMap0 = Collection0.AddSetMap<FFloatAnimationAttribute>(FloatToFloatKey);
+			UE_RETURN_ON_ERROR(FloatToFloatMap0 != nullptr, TEXT("Float->Float map should be valid"));
+
+			FAttributeMappingKey FloatToIntegerKey = FAttributeMappingKey::MakeFromTo<FFloatAnimationAttribute, FIntegerAnimationAttribute>();
+			TAttributeSetMap<FIntegerAnimationAttribute>* FloatToIntegerMap0 = Collection0.AddSetMap<FIntegerAnimationAttribute>(FloatToIntegerKey);
+			UE_RETURN_ON_ERROR(FloatToIntegerMap0 != nullptr, TEXT("Float->Integer map should be valid"));
+
+			FAttributeMappingKey FloatToStringKey = FAttributeMappingKey::MakeFromTo<FFloatAnimationAttribute, FStringAnimationAttribute>();
+			TAttributeSetMap<FStringAnimationAttribute>* FloatToStringMap0 = Collection0.AddSetMap<FStringAnimationAttribute>(FloatToStringKey);
+			UE_RETURN_ON_ERROR(FloatToStringMap0 != nullptr, TEXT("Float->String map should be valid"));
+
+			FAttributeMappingKey TransformToTransformKey = FAttributeMappingKey::MakeFromTo<FBoneTransformAnimationAttribute, FBoneTransformAnimationAttribute>();
+			TAttributeSetMap<FBoneTransformAnimationAttribute>* TransformToTransformMap0 = Collection0.AddSetMap<FBoneTransformAnimationAttribute>(TransformToTransformKey);
+			UE_RETURN_ON_ERROR(TransformToTransformMap0 != nullptr, TEXT("Transform->Transform map should be valid"));
+
+			FAttributeMappingKey TransformToFloatKey = FAttributeMappingKey::MakeFromTo<FBoneTransformAnimationAttribute, FFloatAnimationAttribute>();
+			TAttributeSetMap<FFloatAnimationAttribute>* TransformToFloatMap0 = Collection0.AddSetMap<FFloatAnimationAttribute>(TransformToFloatKey);
+			UE_RETURN_ON_ERROR(TransformToFloatMap0 != nullptr, TEXT("Transform->Bone map should be valid"));
+
+			TAttributeMap<FFloatAnimationAttribute>* FloatMap0 = Collection0.AddMap<FFloatAnimationAttribute>();
+			UE_RETURN_ON_ERROR(FloatMap0 != nullptr, TEXT("Float map should be valid"));
+
+			TAttributeMap<FBoneTransformAnimationAttribute>* BoneMap0 = Collection0.AddMap<FBoneTransformAnimationAttribute>();
+			UE_RETURN_ON_ERROR(BoneMap0 != nullptr, TEXT("Bone map should be valid"));
+
+			AddErrorIfFalse(!Collection0.IsEmpty(), TEXT("Attribute collection should not be empty"));
+			AddErrorIfFalse(Collection0.NumSetMaps() == 5, TEXT("Unexpected set map count"));
+			AddErrorIfFalse(Collection0.NumMaps() == 2, TEXT("Unexpected set count"));
+
+			FAttributeCollectionHeap Collection1(FullBodySet);
+
+			AddErrorIfFalse(Collection1.IsEmpty(), TEXT("Attribute collection should be empty"));
+			AddErrorIfFalse(Collection1.GetNamedSet() == FullBodySet, TEXT("Unexpected named set"));
+
+			TAttributeSetMap<FFloatAnimationAttribute>* FloatToFloatMap1 = Collection1.AddSetMap<FFloatAnimationAttribute>(FloatToFloatKey);
+			UE_RETURN_ON_ERROR(FloatToFloatMap1 != nullptr, TEXT("Float->Float map should be valid"));
+
+			TAttributeSetMap<FIntegerAnimationAttribute>* FloatToIntegerMap1 = Collection1.AddSetMap<FIntegerAnimationAttribute>(FloatToIntegerKey);
+			UE_RETURN_ON_ERROR(FloatToIntegerMap1 != nullptr, TEXT("Float->Integer map should be valid"));
+
+			TAttributeSetMap<FStringAnimationAttribute>* FloatToStringMap1 = Collection1.AddSetMap<FStringAnimationAttribute>(FloatToStringKey);
+			UE_RETURN_ON_ERROR(FloatToStringMap1 != nullptr, TEXT("Float->String map should be valid"));
+
+			TAttributeSetMap<FBoneTransformAnimationAttribute>* TransformToTransformMap1 = Collection1.AddSetMap<FBoneTransformAnimationAttribute>(TransformToTransformKey);
+			UE_RETURN_ON_ERROR(TransformToTransformMap1 != nullptr, TEXT("Transform->Transform map should be valid"));
+
+			TAttributeSetMap<FFloatAnimationAttribute>* TransformToFloatMap1 = Collection1.AddSetMap<FFloatAnimationAttribute>(TransformToFloatKey);
+			UE_RETURN_ON_ERROR(TransformToFloatMap1 != nullptr, TEXT("Transform->Bone map should be valid"));
+
+			TAttributeMap<FFloatAnimationAttribute>* FloatMap1 = Collection1.AddMap<FFloatAnimationAttribute>();
+			UE_RETURN_ON_ERROR(FloatMap1 != nullptr, TEXT("Float map should be valid"));
+
+			TAttributeMap<FBoneTransformAnimationAttribute>* BoneMap1 = Collection1.AddMap<FBoneTransformAnimationAttribute>();
+			UE_RETURN_ON_ERROR(BoneMap1 != nullptr, TEXT("Bone map should be valid"));
+
+			AddErrorIfFalse(!Collection1.IsEmpty(), TEXT("Attribute collection should not be empty"));
+			AddErrorIfFalse(Collection1.NumSetMaps() == 5, TEXT("Unexpected set map count"));
+			AddErrorIfFalse(Collection1.NumMaps() == 2, TEXT("Unexpected set count"));
+
+			// Both inputs perfectly match (set maps)
+			{
+				int32 NumPredicateCalls = 0;
+				InnerJoinBy(
+					FAttributeCollectionSetMapJoinOp(),
+					[this, &NumPredicateCalls, &Collection0, &Collection1](FAttributeSetMap* Map0, FAttributeSetMap* Map1)
+					{
+						NumPredicateCalls++;
+						AddErrorIfFalse(CollectionContains(Collection0, Map0), TEXT("Unexpected InnerJoin iteration result"));
+						AddErrorIfFalse(CollectionContains(Collection1, Map1), TEXT("Unexpected InnerJoin iteration result"));
+						AddErrorIfFalse(Map0 != nullptr && Map1 != nullptr && Map0->GetAttributeType() == Map1->GetAttributeType(), TEXT("Unexpected InnerJoin iteration result"));
+						AddErrorIfFalse(Map0 != nullptr && Map1 != nullptr && Map0->GetValueType() == Map1->GetValueType(), TEXT("Unexpected InnerJoin iteration result"));
+					},
+					Collection0.CreateSetMapIterator(),
+					Collection1.CreateSetMapIterator());
+				AddErrorIfFalse(NumPredicateCalls == 5, TEXT("Unexpected InnerJoin iteration result"));
+
+				NumPredicateCalls = 0;
+				InnerJoinBy(
+					FAttributeCollectionSetMapJoinOp(),
+					[this, &NumPredicateCalls, &Collection0, &Collection1](const FAttributeSetMap* Map0, const FAttributeSetMap* Map1)
+					{
+						NumPredicateCalls++;
+						AddErrorIfFalse(CollectionContains(Collection0, Map0), TEXT("Unexpected InnerJoin iteration result"));
+						AddErrorIfFalse(CollectionContains(Collection1, Map1), TEXT("Unexpected InnerJoin iteration result"));
+						AddErrorIfFalse(Map0 != nullptr && Map1 != nullptr && Map0->GetAttributeType() == Map1->GetAttributeType(), TEXT("Unexpected InnerJoin iteration result"));
+						AddErrorIfFalse(Map0 != nullptr && Map1 != nullptr && Map0->GetValueType() == Map1->GetValueType(), TEXT("Unexpected InnerJoin iteration result"));
+					},
+					Collection0.CreateConstSetMapIterator(),
+					Collection1.CreateConstSetMapIterator());
+				AddErrorIfFalse(NumPredicateCalls == 5, TEXT("Unexpected InnerJoin iteration result"));
+
+				NumPredicateCalls = 0;
+				InnerJoinBy(
+					FAttributeCollectionSetMapJoinOp(),
+					[this, &NumPredicateCalls, &Collection0, &Collection1](const FAttributeSetMap* Map0, FAttributeSetMap* Map1)
+					{
+						NumPredicateCalls++;
+						AddErrorIfFalse(CollectionContains(Collection0, Map0), TEXT("Unexpected InnerJoin iteration result"));
+						AddErrorIfFalse(CollectionContains(Collection1, Map1), TEXT("Unexpected InnerJoin iteration result"));
+						AddErrorIfFalse(Map0 != nullptr && Map1 != nullptr && Map0->GetAttributeType() == Map1->GetAttributeType(), TEXT("Unexpected InnerJoin iteration result"));
+						AddErrorIfFalse(Map0 != nullptr && Map1 != nullptr && Map0->GetValueType() == Map1->GetValueType(), TEXT("Unexpected InnerJoin iteration result"));
+					},
+					Collection0.CreateConstSetMapIterator(),
+					Collection1.CreateSetMapIterator());
+				AddErrorIfFalse(NumPredicateCalls == 5, TEXT("Unexpected InnerJoin iteration result"));
+
+				NumPredicateCalls = 0;
+				InnerJoinBy(
+					FAttributeCollectionSetMapJoinOp(),
+					[this, &NumPredicateCalls, &Collection0, &Collection1](FAttributeSetMap* Map0, const FAttributeSetMap* Map1)
+					{
+						NumPredicateCalls++;
+						AddErrorIfFalse(CollectionContains(Collection0, Map0), TEXT("Unexpected InnerJoin iteration result"));
+						AddErrorIfFalse(CollectionContains(Collection1, Map1), TEXT("Unexpected InnerJoin iteration result"));
+						AddErrorIfFalse(Map0 != nullptr && Map1 != nullptr && Map0->GetAttributeType() == Map1->GetAttributeType(), TEXT("Unexpected InnerJoin iteration result"));
+						AddErrorIfFalse(Map0 != nullptr && Map1 != nullptr && Map0->GetValueType() == Map1->GetValueType(), TEXT("Unexpected InnerJoin iteration result"));
+					},
+					Collection0.CreateSetMapIterator(),
+					Collection1.CreateConstSetMapIterator());
+				AddErrorIfFalse(NumPredicateCalls == 5, TEXT("Unexpected InnerJoin iteration result"));
+			}
+
+			// Both inputs perfectly match (name/value maps)
+			{
+				int32 NumPredicateCalls = 0;
+				InnerJoinBy(
+					FAttributeCollectionMapJoinOp(),
+					[this, &NumPredicateCalls, &Collection0, &Collection1](FAttributeMap* Map0, FAttributeMap* Map1)
+					{
+						NumPredicateCalls++;
+						AddErrorIfFalse(CollectionContains(Collection0, Map0), TEXT("Unexpected InnerJoin iteration result"));
+						AddErrorIfFalse(CollectionContains(Collection1, Map1), TEXT("Unexpected InnerJoin iteration result"));
+						AddErrorIfFalse(Map0 != nullptr && Map1 != nullptr && Map0->GetValueType() == Map1->GetValueType(), TEXT("Unexpected InnerJoin iteration result"));
+					},
+					Collection0.CreateMapIterator(),
+					Collection1.CreateMapIterator());
+				AddErrorIfFalse(NumPredicateCalls == 2, TEXT("Unexpected InnerJoin iteration result"));
+
+				NumPredicateCalls = 0;
+				InnerJoinBy(
+					FAttributeCollectionMapJoinOp(),
+					[this, &NumPredicateCalls, &Collection0, &Collection1](const FAttributeMap* Map0, const FAttributeMap* Map1)
+					{
+						NumPredicateCalls++;
+						AddErrorIfFalse(CollectionContains(Collection0, Map0), TEXT("Unexpected InnerJoin iteration result"));
+						AddErrorIfFalse(CollectionContains(Collection1, Map1), TEXT("Unexpected InnerJoin iteration result"));
+						AddErrorIfFalse(Map0 != nullptr && Map1 != nullptr && Map0->GetValueType() == Map1->GetValueType(), TEXT("Unexpected InnerJoin iteration result"));
+					},
+					Collection0.CreateConstMapIterator(),
+					Collection1.CreateConstMapIterator());
+				AddErrorIfFalse(NumPredicateCalls == 2, TEXT("Unexpected InnerJoin iteration result"));
+
+				NumPredicateCalls = 0;
+				InnerJoinBy(
+					FAttributeCollectionMapJoinOp(),
+					[this, &NumPredicateCalls, &Collection0, &Collection1](const FAttributeMap* Map0, FAttributeMap* Map1)
+					{
+						NumPredicateCalls++;
+						AddErrorIfFalse(CollectionContains(Collection0, Map0), TEXT("Unexpected InnerJoin iteration result"));
+						AddErrorIfFalse(CollectionContains(Collection1, Map1), TEXT("Unexpected InnerJoin iteration result"));
+						AddErrorIfFalse(Map0 != nullptr && Map1 != nullptr && Map0->GetValueType() == Map1->GetValueType(), TEXT("Unexpected InnerJoin iteration result"));
+					},
+					Collection0.CreateConstMapIterator(),
+					Collection1.CreateMapIterator());
+				AddErrorIfFalse(NumPredicateCalls == 2, TEXT("Unexpected InnerJoin iteration result"));
+
+				NumPredicateCalls = 0;
+				InnerJoinBy(
+					FAttributeCollectionMapJoinOp(),
+					[this, &NumPredicateCalls, &Collection0, &Collection1](FAttributeMap* Map0, const FAttributeMap* Map1)
+					{
+						NumPredicateCalls++;
+						AddErrorIfFalse(CollectionContains(Collection0, Map0), TEXT("Unexpected InnerJoin iteration result"));
+						AddErrorIfFalse(CollectionContains(Collection1, Map1), TEXT("Unexpected InnerJoin iteration result"));
+						AddErrorIfFalse(Map0 != nullptr && Map1 != nullptr && Map0->GetValueType() == Map1->GetValueType(), TEXT("Unexpected InnerJoin iteration result"));
+					},
+					Collection0.CreateMapIterator(),
+					Collection1.CreateConstMapIterator());
+				AddErrorIfFalse(NumPredicateCalls == 2, TEXT("Unexpected InnerJoin iteration result"));
+			}
+
+			FAttributeCollectionHeap Collection2(FullBodySet);
+
+			AddErrorIfFalse(Collection2.IsEmpty(), TEXT("Attribute collection should be empty"));
+			AddErrorIfFalse(Collection2.GetNamedSet() == FullBodySet, TEXT("Unexpected named set"));
+
+			TAttributeSetMap<FFloatAnimationAttribute>* FloatToFloatMap2 = Collection2.AddSetMap<FFloatAnimationAttribute>(FloatToFloatKey);
+			UE_RETURN_ON_ERROR(FloatToFloatMap2 != nullptr, TEXT("Float->Float map should be valid"));
+
+			TAttributeSetMap<FBoneTransformAnimationAttribute>* TransformToTransformMap2 = Collection2.AddSetMap<FBoneTransformAnimationAttribute>(TransformToTransformKey);
+			UE_RETURN_ON_ERROR(TransformToTransformMap2 != nullptr, TEXT("Transform->Transform map should be valid"));
+
+			TAttributeMap<FFloatAnimationAttribute>* FloatMap2 = Collection2.AddMap<FFloatAnimationAttribute>();
+			UE_RETURN_ON_ERROR(FloatMap2 != nullptr, TEXT("Float map should be valid"));
+
+			AddErrorIfFalse(!Collection2.IsEmpty(), TEXT("Attribute collection should not be empty"));
+			AddErrorIfFalse(Collection2.NumSetMaps() == 2, TEXT("Unexpected set map count"));
+			AddErrorIfFalse(Collection2.NumMaps() == 1, TEXT("Unexpected map count"));
+
+			// First input is shorter and its entries are skipped when missing (set maps)
+			{
+				int32 NumPredicateCalls = 0;
+				InnerJoinBy(
+					FAttributeCollectionSetMapJoinOp(),
+					[this, &NumPredicateCalls, &Collection0, &Collection2](FAttributeSetMap* Map2, FAttributeSetMap* Map0)
+					{
+						NumPredicateCalls++;
+						AddErrorIfFalse(CollectionContains(Collection0, Map0), TEXT("Unexpected InnerJoin iteration result"));
+						AddErrorIfFalse(CollectionContains(Collection2, Map2), TEXT("Unexpected InnerJoin iteration result"));
+						AddErrorIfFalse(Map0 != nullptr && Map2 != nullptr && Map0->GetAttributeType() == Map2->GetAttributeType(), TEXT("Unexpected InnerJoin iteration result"));
+						AddErrorIfFalse(Map0 != nullptr && Map2 != nullptr && Map0->GetValueType() == Map2->GetValueType(), TEXT("Unexpected InnerJoin iteration result"));
+					},
+					Collection2.CreateSetMapIterator(),
+					Collection0.CreateSetMapIterator());
+				AddErrorIfFalse(NumPredicateCalls == 2, TEXT("Unexpected InnerJoin iteration result"));
+
+				NumPredicateCalls = 0;
+				InnerJoinBy(
+					FAttributeCollectionSetMapJoinOp(),
+					[this, &NumPredicateCalls, &Collection0, &Collection2](const FAttributeSetMap* Map2, const FAttributeSetMap* Map0)
+					{
+						NumPredicateCalls++;
+						AddErrorIfFalse(CollectionContains(Collection0, Map0), TEXT("Unexpected InnerJoin iteration result"));
+						AddErrorIfFalse(CollectionContains(Collection2, Map2), TEXT("Unexpected InnerJoin iteration result"));
+						AddErrorIfFalse(Map0 != nullptr && Map2 != nullptr && Map0->GetAttributeType() == Map2->GetAttributeType(), TEXT("Unexpected InnerJoin iteration result"));
+						AddErrorIfFalse(Map0 != nullptr && Map2 != nullptr && Map0->GetValueType() == Map2->GetValueType(), TEXT("Unexpected InnerJoin iteration result"));
+					},
+					Collection2.CreateConstSetMapIterator(),
+					Collection0.CreateConstSetMapIterator());
+				AddErrorIfFalse(NumPredicateCalls == 2, TEXT("Unexpected InnerJoin iteration result"));
+
+				NumPredicateCalls = 0;
+				InnerJoinBy(
+					FAttributeCollectionSetMapJoinOp(),
+					[this, &NumPredicateCalls, &Collection0, &Collection2](const FAttributeSetMap* Map2, FAttributeSetMap* Map0)
+					{
+						NumPredicateCalls++;
+						AddErrorIfFalse(CollectionContains(Collection0, Map0), TEXT("Unexpected InnerJoin iteration result"));
+						AddErrorIfFalse(CollectionContains(Collection2, Map2), TEXT("Unexpected InnerJoin iteration result"));
+						AddErrorIfFalse(Map0 != nullptr && Map2 != nullptr && Map0->GetAttributeType() == Map2->GetAttributeType(), TEXT("Unexpected InnerJoin iteration result"));
+						AddErrorIfFalse(Map0 != nullptr && Map2 != nullptr && Map0->GetValueType() == Map2->GetValueType(), TEXT("Unexpected InnerJoin iteration result"));
+					},
+					Collection2.CreateConstSetMapIterator(),
+					Collection0.CreateSetMapIterator());
+				AddErrorIfFalse(NumPredicateCalls == 2, TEXT("Unexpected InnerJoin iteration result"));
+
+				NumPredicateCalls = 0;
+				InnerJoinBy(
+					FAttributeCollectionSetMapJoinOp(),
+					[this, &NumPredicateCalls, &Collection0, &Collection2](FAttributeSetMap* Map2, const FAttributeSetMap* Map0)
+					{
+						NumPredicateCalls++;
+						AddErrorIfFalse(CollectionContains(Collection0, Map0), TEXT("Unexpected InnerJoin iteration result"));
+						AddErrorIfFalse(CollectionContains(Collection2, Map2), TEXT("Unexpected InnerJoin iteration result"));
+						AddErrorIfFalse(Map0 != nullptr && Map2 != nullptr && Map0->GetAttributeType() == Map2->GetAttributeType(), TEXT("Unexpected InnerJoin iteration result"));
+						AddErrorIfFalse(Map0 != nullptr && Map2 != nullptr && Map0->GetValueType() == Map2->GetValueType(), TEXT("Unexpected InnerJoin iteration result"));
+					},
+					Collection2.CreateSetMapIterator(),
+					Collection0.CreateConstSetMapIterator());
+				AddErrorIfFalse(NumPredicateCalls == 2, TEXT("Unexpected InnerJoin iteration result"));
+			}
+
+			// First input is shorter and its entries are skipped when missing (name/value maps)
+			{
+				int32 NumPredicateCalls = 0;
+				InnerJoinBy(
+					FAttributeCollectionMapJoinOp(),
+					[this, &NumPredicateCalls, &Collection0, &Collection2](FAttributeMap* Map2, FAttributeMap* Map0)
+					{
+						NumPredicateCalls++;
+						AddErrorIfFalse(CollectionContains(Collection0, Map0), TEXT("Unexpected InnerJoin iteration result"));
+						AddErrorIfFalse(CollectionContains(Collection2, Map2), TEXT("Unexpected InnerJoin iteration result"));
+						AddErrorIfFalse(Map0 != nullptr && Map2 != nullptr && Map0->GetValueType() == Map2->GetValueType(), TEXT("Unexpected InnerJoin iteration result"));
+					},
+					Collection2.CreateMapIterator(),
+					Collection0.CreateMapIterator());
+				AddErrorIfFalse(NumPredicateCalls == 1, TEXT("Unexpected InnerJoin iteration result"));
+
+				NumPredicateCalls = 0;
+				InnerJoinBy(
+					FAttributeCollectionMapJoinOp(),
+					[this, &NumPredicateCalls, &Collection0, &Collection2](const FAttributeMap* Map2, const FAttributeMap* Map0)
+					{
+						NumPredicateCalls++;
+						AddErrorIfFalse(CollectionContains(Collection0, Map0), TEXT("Unexpected InnerJoin iteration result"));
+						AddErrorIfFalse(CollectionContains(Collection2, Map2), TEXT("Unexpected InnerJoin iteration result"));
+						AddErrorIfFalse(Map0 != nullptr && Map2 != nullptr && Map0->GetValueType() == Map2->GetValueType(), TEXT("Unexpected InnerJoin iteration result"));
+					},
+					Collection2.CreateConstMapIterator(),
+					Collection0.CreateConstMapIterator());
+				AddErrorIfFalse(NumPredicateCalls == 1, TEXT("Unexpected InnerJoin iteration result"));
+
+				NumPredicateCalls = 0;
+				InnerJoinBy(
+					FAttributeCollectionMapJoinOp(),
+					[this, &NumPredicateCalls, &Collection0, &Collection2](const FAttributeMap* Map2, FAttributeMap* Map0)
+					{
+						NumPredicateCalls++;
+						AddErrorIfFalse(CollectionContains(Collection0, Map0), TEXT("Unexpected InnerJoin iteration result"));
+						AddErrorIfFalse(CollectionContains(Collection2, Map2), TEXT("Unexpected InnerJoin iteration result"));
+						AddErrorIfFalse(Map0 != nullptr && Map2 != nullptr && Map0->GetValueType() == Map2->GetValueType(), TEXT("Unexpected InnerJoin iteration result"));
+					},
+					Collection2.CreateConstMapIterator(),
+					Collection0.CreateMapIterator());
+				AddErrorIfFalse(NumPredicateCalls == 1, TEXT("Unexpected InnerJoin iteration result"));
+
+				NumPredicateCalls = 0;
+				InnerJoinBy(
+					FAttributeCollectionMapJoinOp(),
+					[this, &NumPredicateCalls, &Collection0, &Collection2](FAttributeMap* Map2, const FAttributeMap* Map0)
+					{
+						NumPredicateCalls++;
+						AddErrorIfFalse(CollectionContains(Collection0, Map0), TEXT("Unexpected InnerJoin iteration result"));
+						AddErrorIfFalse(CollectionContains(Collection2, Map2), TEXT("Unexpected InnerJoin iteration result"));
+						AddErrorIfFalse(Map0 != nullptr && Map2 != nullptr && Map0->GetValueType() == Map2->GetValueType(), TEXT("Unexpected InnerJoin iteration result"));
+					},
+					Collection2.CreateMapIterator(),
+					Collection0.CreateConstMapIterator());
+				AddErrorIfFalse(NumPredicateCalls == 1, TEXT("Unexpected InnerJoin iteration result"));
+			}
+
+			// Second input is shorter and its entries are skipped when missing (set maps)
+			{
+				int32 NumPredicateCalls = 0;
+				InnerJoinBy(
+					FAttributeCollectionSetMapJoinOp(),
+					[this, &NumPredicateCalls, &Collection0, &Collection2](FAttributeSetMap* Map0, FAttributeSetMap* Map2)
+					{
+						NumPredicateCalls++;
+						AddErrorIfFalse(CollectionContains(Collection0, Map0), TEXT("Unexpected InnerJoin iteration result"));
+						AddErrorIfFalse(CollectionContains(Collection2, Map2), TEXT("Unexpected InnerJoin iteration result"));
+						AddErrorIfFalse(Map0 != nullptr && Map2 != nullptr && Map0->GetAttributeType() == Map2->GetAttributeType(), TEXT("Unexpected InnerJoin iteration result"));
+						AddErrorIfFalse(Map0 != nullptr && Map2 != nullptr && Map0->GetValueType() == Map2->GetValueType(), TEXT("Unexpected InnerJoin iteration result"));
+					},
+					Collection0.CreateSetMapIterator(),
+					Collection2.CreateSetMapIterator());
+				AddErrorIfFalse(NumPredicateCalls == 2, TEXT("Unexpected InnerJoin iteration result"));
+
+				NumPredicateCalls = 0;
+				InnerJoinBy(
+					FAttributeCollectionSetMapJoinOp(),
+					[this, &NumPredicateCalls, &Collection0, &Collection2](const FAttributeSetMap* Map0, const FAttributeSetMap* Map2)
+					{
+						NumPredicateCalls++;
+						AddErrorIfFalse(CollectionContains(Collection0, Map0), TEXT("Unexpected InnerJoin iteration result"));
+						AddErrorIfFalse(CollectionContains(Collection2, Map2), TEXT("Unexpected InnerJoin iteration result"));
+						AddErrorIfFalse(Map0 != nullptr && Map2 != nullptr && Map0->GetAttributeType() == Map2->GetAttributeType(), TEXT("Unexpected InnerJoin iteration result"));
+						AddErrorIfFalse(Map0 != nullptr && Map2 != nullptr && Map0->GetValueType() == Map2->GetValueType(), TEXT("Unexpected InnerJoin iteration result"));
+					},
+					Collection0.CreateConstSetMapIterator(),
+					Collection2.CreateConstSetMapIterator());
+				AddErrorIfFalse(NumPredicateCalls == 2, TEXT("Unexpected InnerJoin iteration result"));
+
+				NumPredicateCalls = 0;
+				InnerJoinBy(
+					FAttributeCollectionSetMapJoinOp(),
+					[this, &NumPredicateCalls, &Collection0, &Collection2](const FAttributeSetMap* Map0, FAttributeSetMap* Map2)
+					{
+						NumPredicateCalls++;
+						AddErrorIfFalse(CollectionContains(Collection0, Map0), TEXT("Unexpected InnerJoin iteration result"));
+						AddErrorIfFalse(CollectionContains(Collection2, Map2), TEXT("Unexpected InnerJoin iteration result"));
+						AddErrorIfFalse(Map0 != nullptr && Map2 != nullptr && Map0->GetAttributeType() == Map2->GetAttributeType(), TEXT("Unexpected InnerJoin iteration result"));
+						AddErrorIfFalse(Map0 != nullptr && Map2 != nullptr && Map0->GetValueType() == Map2->GetValueType(), TEXT("Unexpected InnerJoin iteration result"));
+					},
+					Collection0.CreateConstSetMapIterator(),
+					Collection2.CreateSetMapIterator());
+				AddErrorIfFalse(NumPredicateCalls == 2, TEXT("Unexpected InnerJoin iteration result"));
+
+				NumPredicateCalls = 0;
+				InnerJoinBy(
+					FAttributeCollectionSetMapJoinOp(),
+					[this, &NumPredicateCalls, &Collection0, &Collection2](FAttributeSetMap* Map0, const FAttributeSetMap* Map2)
+					{
+						NumPredicateCalls++;
+						AddErrorIfFalse(CollectionContains(Collection0, Map0), TEXT("Unexpected InnerJoin iteration result"));
+						AddErrorIfFalse(CollectionContains(Collection2, Map2), TEXT("Unexpected InnerJoin iteration result"));
+						AddErrorIfFalse(Map0 != nullptr && Map2 != nullptr && Map0->GetAttributeType() == Map2->GetAttributeType(), TEXT("Unexpected InnerJoin iteration result"));
+						AddErrorIfFalse(Map0 != nullptr && Map2 != nullptr && Map0->GetValueType() == Map2->GetValueType(), TEXT("Unexpected InnerJoin iteration result"));
+					},
+					Collection0.CreateSetMapIterator(),
+					Collection2.CreateConstSetMapIterator());
+				AddErrorIfFalse(NumPredicateCalls == 2, TEXT("Unexpected InnerJoin iteration result"));
+			}
+
+			// Second input is shorter and its entries are skipped when missing (name/value maps)
+			{
+				int32 NumPredicateCalls = 0;
+				InnerJoinBy(
+					FAttributeCollectionMapJoinOp(),
+					[this, &NumPredicateCalls, &Collection0, &Collection2](FAttributeMap* Map0, FAttributeMap* Map2)
+					{
+						NumPredicateCalls++;
+						AddErrorIfFalse(CollectionContains(Collection0, Map0), TEXT("Unexpected InnerJoin iteration result"));
+						AddErrorIfFalse(CollectionContains(Collection2, Map2), TEXT("Unexpected InnerJoin iteration result"));
+						AddErrorIfFalse(Map0 != nullptr && Map2 != nullptr && Map0->GetValueType() == Map2->GetValueType(), TEXT("Unexpected InnerJoin iteration result"));
+					},
+					Collection0.CreateMapIterator(),
+					Collection2.CreateMapIterator());
+				AddErrorIfFalse(NumPredicateCalls == 1, TEXT("Unexpected InnerJoin iteration result"));
+
+				NumPredicateCalls = 0;
+				InnerJoinBy(
+					FAttributeCollectionMapJoinOp(),
+					[this, &NumPredicateCalls, &Collection0, &Collection2](const FAttributeMap* Map0, const FAttributeMap* Map2)
+					{
+						NumPredicateCalls++;
+						AddErrorIfFalse(CollectionContains(Collection0, Map0), TEXT("Unexpected InnerJoin iteration result"));
+						AddErrorIfFalse(CollectionContains(Collection2, Map2), TEXT("Unexpected InnerJoin iteration result"));
+						AddErrorIfFalse(Map0 != nullptr && Map2 != nullptr && Map0->GetValueType() == Map2->GetValueType(), TEXT("Unexpected InnerJoin iteration result"));
+					},
+					Collection0.CreateConstMapIterator(),
+					Collection2.CreateConstMapIterator());
+				AddErrorIfFalse(NumPredicateCalls == 1, TEXT("Unexpected InnerJoin iteration result"));
+
+				NumPredicateCalls = 0;
+				InnerJoinBy(
+					FAttributeCollectionMapJoinOp(),
+					[this, &NumPredicateCalls, &Collection0, &Collection2](const FAttributeMap* Map0, FAttributeMap* Map2)
+					{
+						NumPredicateCalls++;
+						AddErrorIfFalse(CollectionContains(Collection0, Map0), TEXT("Unexpected InnerJoin iteration result"));
+						AddErrorIfFalse(CollectionContains(Collection2, Map2), TEXT("Unexpected InnerJoin iteration result"));
+						AddErrorIfFalse(Map0 != nullptr && Map2 != nullptr && Map0->GetValueType() == Map2->GetValueType(), TEXT("Unexpected InnerJoin iteration result"));
+					},
+					Collection0.CreateConstMapIterator(),
+					Collection2.CreateMapIterator());
+				AddErrorIfFalse(NumPredicateCalls == 1, TEXT("Unexpected InnerJoin iteration result"));
+
+				NumPredicateCalls = 0;
+				InnerJoinBy(
+					FAttributeCollectionMapJoinOp(),
+					[this, &NumPredicateCalls, &Collection0, &Collection2](FAttributeMap* Map0, const FAttributeMap* Map2)
+					{
+						NumPredicateCalls++;
+						AddErrorIfFalse(CollectionContains(Collection0, Map0), TEXT("Unexpected InnerJoin iteration result"));
+						AddErrorIfFalse(CollectionContains(Collection2, Map2), TEXT("Unexpected InnerJoin iteration result"));
+						AddErrorIfFalse(Map0 != nullptr && Map2 != nullptr && Map0->GetValueType() == Map2->GetValueType(), TEXT("Unexpected InnerJoin iteration result"));
+					},
+					Collection0.CreateMapIterator(),
+					Collection2.CreateConstMapIterator());
+				AddErrorIfFalse(NumPredicateCalls == 1, TEXT("Unexpected InnerJoin iteration result"));
+			}
+
+			// 3 input variant (set maps)
+			{
+				int32 NumPredicateCalls = 0;
+				InnerJoinBy(
+					FAttributeCollectionSetMapJoinOp(),
+					[this, &NumPredicateCalls, &Collection0, &Collection1, &Collection2](FAttributeSetMap* Map0, FAttributeSetMap* Map1, FAttributeSetMap* Map2)
+					{
+						NumPredicateCalls++;
+						AddErrorIfFalse(CollectionContains(Collection0, Map0), TEXT("Unexpected InnerJoin iteration result"));
+						AddErrorIfFalse(CollectionContains(Collection1, Map1), TEXT("Unexpected InnerJoin iteration result"));
+						AddErrorIfFalse(CollectionContains(Collection2, Map2), TEXT("Unexpected InnerJoin iteration result"));
+						AddErrorIfFalse(Map0 != nullptr && Map1 != nullptr && Map0->GetAttributeType() == Map1->GetAttributeType(), TEXT("Unexpected InnerJoin iteration result"));
+						AddErrorIfFalse(Map0 != nullptr && Map2 != nullptr && Map0->GetAttributeType() == Map2->GetAttributeType(), TEXT("Unexpected InnerJoin iteration result"));
+						AddErrorIfFalse(Map0 != nullptr && Map1 != nullptr && Map0->GetValueType() == Map1->GetValueType(), TEXT("Unexpected InnerJoin iteration result"));
+						AddErrorIfFalse(Map0 != nullptr && Map2 != nullptr && Map0->GetValueType() == Map2->GetValueType(), TEXT("Unexpected InnerJoin iteration result"));
+					},
+					Collection0.CreateSetMapIterator(),
+					Collection1.CreateSetMapIterator(),
+					Collection2.CreateSetMapIterator());
+				AddErrorIfFalse(NumPredicateCalls == 2, TEXT("Unexpected InnerJoin iteration result"));
+
+				NumPredicateCalls = 0;
+				InnerJoinBy(
+					FAttributeCollectionSetMapJoinOp(),
+					[this, &NumPredicateCalls, &Collection0, &Collection1, &Collection2](FAttributeSetMap* Map0, FAttributeSetMap* Map2, FAttributeSetMap* Map1)
+					{
+						NumPredicateCalls++;
+						AddErrorIfFalse(CollectionContains(Collection0, Map0), TEXT("Unexpected InnerJoin iteration result"));
+						AddErrorIfFalse(CollectionContains(Collection1, Map1), TEXT("Unexpected InnerJoin iteration result"));
+						AddErrorIfFalse(CollectionContains(Collection2, Map2), TEXT("Unexpected InnerJoin iteration result"));
+						AddErrorIfFalse(Map0 != nullptr && Map1 != nullptr && Map0->GetAttributeType() == Map1->GetAttributeType(), TEXT("Unexpected InnerJoin iteration result"));
+						AddErrorIfFalse(Map0 != nullptr && Map2 != nullptr && Map0->GetAttributeType() == Map2->GetAttributeType(), TEXT("Unexpected InnerJoin iteration result"));
+						AddErrorIfFalse(Map0 != nullptr && Map1 != nullptr && Map0->GetValueType() == Map1->GetValueType(), TEXT("Unexpected InnerJoin iteration result"));
+						AddErrorIfFalse(Map0 != nullptr && Map2 != nullptr && Map0->GetValueType() == Map2->GetValueType(), TEXT("Unexpected InnerJoin iteration result"));
+					},
+					Collection0.CreateSetMapIterator(),
+					Collection2.CreateSetMapIterator(),
+					Collection1.CreateSetMapIterator());
+				AddErrorIfFalse(NumPredicateCalls == 2, TEXT("Unexpected InnerJoin iteration result"));
+
+				NumPredicateCalls = 0;
+				InnerJoinBy(
+					FAttributeCollectionSetMapJoinOp(),
+					[this, &NumPredicateCalls, &Collection0, &Collection1, &Collection2](FAttributeSetMap* Map1, FAttributeSetMap* Map0, FAttributeSetMap* Map2)
+					{
+						NumPredicateCalls++;
+						AddErrorIfFalse(CollectionContains(Collection0, Map0), TEXT("Unexpected InnerJoin iteration result"));
+						AddErrorIfFalse(CollectionContains(Collection1, Map1), TEXT("Unexpected InnerJoin iteration result"));
+						AddErrorIfFalse(CollectionContains(Collection2, Map2), TEXT("Unexpected InnerJoin iteration result"));
+						AddErrorIfFalse(Map0 != nullptr && Map1 != nullptr && Map0->GetAttributeType() == Map1->GetAttributeType(), TEXT("Unexpected InnerJoin iteration result"));
+						AddErrorIfFalse(Map0 != nullptr && Map2 != nullptr && Map0->GetAttributeType() == Map2->GetAttributeType(), TEXT("Unexpected InnerJoin iteration result"));
+						AddErrorIfFalse(Map0 != nullptr && Map1 != nullptr && Map0->GetValueType() == Map1->GetValueType(), TEXT("Unexpected InnerJoin iteration result"));
+						AddErrorIfFalse(Map0 != nullptr && Map2 != nullptr && Map0->GetValueType() == Map2->GetValueType(), TEXT("Unexpected InnerJoin iteration result"));
+					},
+					Collection1.CreateSetMapIterator(),
+					Collection0.CreateSetMapIterator(),
+					Collection2.CreateSetMapIterator());
+				AddErrorIfFalse(NumPredicateCalls == 2, TEXT("Unexpected InnerJoin iteration result"));
+
+				NumPredicateCalls = 0;
+				InnerJoinBy(
+					FAttributeCollectionSetMapJoinOp(),
+					[this, &NumPredicateCalls, &Collection0, &Collection1, &Collection2](FAttributeSetMap* Map1, FAttributeSetMap* Map2, FAttributeSetMap* Map0)
+					{
+						NumPredicateCalls++;
+						AddErrorIfFalse(CollectionContains(Collection0, Map0), TEXT("Unexpected InnerJoin iteration result"));
+						AddErrorIfFalse(CollectionContains(Collection1, Map1), TEXT("Unexpected InnerJoin iteration result"));
+						AddErrorIfFalse(CollectionContains(Collection2, Map2), TEXT("Unexpected InnerJoin iteration result"));
+						AddErrorIfFalse(Map0 != nullptr && Map1 != nullptr && Map0->GetAttributeType() == Map1->GetAttributeType(), TEXT("Unexpected InnerJoin iteration result"));
+						AddErrorIfFalse(Map0 != nullptr && Map2 != nullptr && Map0->GetAttributeType() == Map2->GetAttributeType(), TEXT("Unexpected InnerJoin iteration result"));
+						AddErrorIfFalse(Map0 != nullptr && Map1 != nullptr && Map0->GetValueType() == Map1->GetValueType(), TEXT("Unexpected InnerJoin iteration result"));
+						AddErrorIfFalse(Map0 != nullptr && Map2 != nullptr && Map0->GetValueType() == Map2->GetValueType(), TEXT("Unexpected InnerJoin iteration result"));
+					},
+					Collection1.CreateSetMapIterator(),
+					Collection2.CreateSetMapIterator(),
+					Collection0.CreateSetMapIterator());
+				AddErrorIfFalse(NumPredicateCalls == 2, TEXT("Unexpected InnerJoin iteration result"));
+
+				NumPredicateCalls = 0;
+				InnerJoinBy(
+					FAttributeCollectionSetMapJoinOp(),
+					[this, &NumPredicateCalls, &Collection0, &Collection1, &Collection2](FAttributeSetMap* Map2, FAttributeSetMap* Map0, FAttributeSetMap* Map1)
+					{
+						NumPredicateCalls++;
+						AddErrorIfFalse(CollectionContains(Collection0, Map0), TEXT("Unexpected InnerJoin iteration result"));
+						AddErrorIfFalse(CollectionContains(Collection1, Map1), TEXT("Unexpected InnerJoin iteration result"));
+						AddErrorIfFalse(CollectionContains(Collection2, Map2), TEXT("Unexpected InnerJoin iteration result"));
+						AddErrorIfFalse(Map0 != nullptr && Map1 != nullptr && Map0->GetAttributeType() == Map1->GetAttributeType(), TEXT("Unexpected InnerJoin iteration result"));
+						AddErrorIfFalse(Map0 != nullptr && Map2 != nullptr && Map0->GetAttributeType() == Map2->GetAttributeType(), TEXT("Unexpected InnerJoin iteration result"));
+						AddErrorIfFalse(Map0 != nullptr && Map1 != nullptr && Map0->GetValueType() == Map1->GetValueType(), TEXT("Unexpected InnerJoin iteration result"));
+						AddErrorIfFalse(Map0 != nullptr && Map2 != nullptr && Map0->GetValueType() == Map2->GetValueType(), TEXT("Unexpected InnerJoin iteration result"));
+					},
+					Collection2.CreateSetMapIterator(),
+					Collection0.CreateSetMapIterator(),
+					Collection1.CreateSetMapIterator());
+				AddErrorIfFalse(NumPredicateCalls == 2, TEXT("Unexpected InnerJoin iteration result"));
+
+				NumPredicateCalls = 0;
+				InnerJoinBy(
+					FAttributeCollectionSetMapJoinOp(),
+					[this, &NumPredicateCalls, &Collection0, &Collection1, &Collection2](FAttributeSetMap* Map2, FAttributeSetMap* Map1, FAttributeSetMap* Map0)
+					{
+						NumPredicateCalls++;
+						AddErrorIfFalse(CollectionContains(Collection0, Map0), TEXT("Unexpected InnerJoin iteration result"));
+						AddErrorIfFalse(CollectionContains(Collection1, Map1), TEXT("Unexpected InnerJoin iteration result"));
+						AddErrorIfFalse(CollectionContains(Collection2, Map2), TEXT("Unexpected InnerJoin iteration result"));
+						AddErrorIfFalse(Map0 != nullptr && Map1 != nullptr && Map0->GetAttributeType() == Map1->GetAttributeType(), TEXT("Unexpected InnerJoin iteration result"));
+						AddErrorIfFalse(Map0 != nullptr && Map2 != nullptr && Map0->GetAttributeType() == Map2->GetAttributeType(), TEXT("Unexpected InnerJoin iteration result"));
+						AddErrorIfFalse(Map0 != nullptr && Map1 != nullptr && Map0->GetValueType() == Map1->GetValueType(), TEXT("Unexpected InnerJoin iteration result"));
+						AddErrorIfFalse(Map0 != nullptr && Map2 != nullptr && Map0->GetValueType() == Map2->GetValueType(), TEXT("Unexpected InnerJoin iteration result"));
+					},
+					Collection2.CreateSetMapIterator(),
+					Collection1.CreateSetMapIterator(),
+					Collection0.CreateSetMapIterator());
+				AddErrorIfFalse(NumPredicateCalls == 2, TEXT("Unexpected InnerJoin iteration result"));
+			}
+
+			// 3 input variant (name/value maps)
+			{
+				int32 NumPredicateCalls = 0;
+				InnerJoinBy(
+					FAttributeCollectionMapJoinOp(),
+					[this, &NumPredicateCalls, &Collection0, &Collection1, &Collection2](FAttributeMap* Map0, FAttributeMap* Map1, FAttributeMap* Map2)
+					{
+						NumPredicateCalls++;
+						AddErrorIfFalse(CollectionContains(Collection0, Map0), TEXT("Unexpected InnerJoin iteration result"));
+						AddErrorIfFalse(CollectionContains(Collection1, Map1), TEXT("Unexpected InnerJoin iteration result"));
+						AddErrorIfFalse(CollectionContains(Collection2, Map2), TEXT("Unexpected InnerJoin iteration result"));
+						AddErrorIfFalse(Map0 != nullptr && Map1 != nullptr && Map0->GetValueType() == Map1->GetValueType(), TEXT("Unexpected InnerJoin iteration result"));
+						AddErrorIfFalse(Map0 != nullptr && Map2 != nullptr && Map0->GetValueType() == Map2->GetValueType(), TEXT("Unexpected InnerJoin iteration result"));
+					},
+					Collection0.CreateMapIterator(),
+					Collection1.CreateMapIterator(),
+					Collection2.CreateMapIterator());
+				AddErrorIfFalse(NumPredicateCalls == 1, TEXT("Unexpected InnerJoin iteration result"));
+
+				NumPredicateCalls = 0;
+				InnerJoinBy(
+					FAttributeCollectionMapJoinOp(),
+					[this, &NumPredicateCalls, &Collection0, &Collection1, &Collection2](FAttributeMap* Map0, FAttributeMap* Map2, FAttributeMap* Map1)
+					{
+						NumPredicateCalls++;
+						AddErrorIfFalse(CollectionContains(Collection0, Map0), TEXT("Unexpected InnerJoin iteration result"));
+						AddErrorIfFalse(CollectionContains(Collection1, Map1), TEXT("Unexpected InnerJoin iteration result"));
+						AddErrorIfFalse(CollectionContains(Collection2, Map2), TEXT("Unexpected InnerJoin iteration result"));
+						AddErrorIfFalse(Map0 != nullptr && Map1 != nullptr && Map0->GetValueType() == Map1->GetValueType(), TEXT("Unexpected InnerJoin iteration result"));
+						AddErrorIfFalse(Map0 != nullptr && Map2 != nullptr && Map0->GetValueType() == Map2->GetValueType(), TEXT("Unexpected InnerJoin iteration result"));
+					},
+					Collection0.CreateMapIterator(),
+					Collection2.CreateMapIterator(),
+					Collection1.CreateMapIterator());
+				AddErrorIfFalse(NumPredicateCalls == 1, TEXT("Unexpected InnerJoin iteration result"));
+
+				NumPredicateCalls = 0;
+				InnerJoinBy(
+					FAttributeCollectionMapJoinOp(),
+					[this, &NumPredicateCalls, &Collection0, &Collection1, &Collection2](FAttributeMap* Map1, FAttributeMap* Map0, FAttributeMap* Map2)
+					{
+						NumPredicateCalls++;
+						AddErrorIfFalse(CollectionContains(Collection0, Map0), TEXT("Unexpected InnerJoin iteration result"));
+						AddErrorIfFalse(CollectionContains(Collection1, Map1), TEXT("Unexpected InnerJoin iteration result"));
+						AddErrorIfFalse(CollectionContains(Collection2, Map2), TEXT("Unexpected InnerJoin iteration result"));
+						AddErrorIfFalse(Map0 != nullptr && Map1 != nullptr && Map0->GetValueType() == Map1->GetValueType(), TEXT("Unexpected InnerJoin iteration result"));
+						AddErrorIfFalse(Map0 != nullptr && Map2 != nullptr && Map0->GetValueType() == Map2->GetValueType(), TEXT("Unexpected InnerJoin iteration result"));
+					},
+					Collection1.CreateMapIterator(),
+					Collection0.CreateMapIterator(),
+					Collection2.CreateMapIterator());
+				AddErrorIfFalse(NumPredicateCalls == 1, TEXT("Unexpected InnerJoin iteration result"));
+
+				NumPredicateCalls = 0;
+				InnerJoinBy(
+					FAttributeCollectionMapJoinOp(),
+					[this, &NumPredicateCalls, &Collection0, &Collection1, &Collection2](FAttributeMap* Map1, FAttributeMap* Map2, FAttributeMap* Map0)
+					{
+						NumPredicateCalls++;
+						AddErrorIfFalse(CollectionContains(Collection0, Map0), TEXT("Unexpected InnerJoin iteration result"));
+						AddErrorIfFalse(CollectionContains(Collection1, Map1), TEXT("Unexpected InnerJoin iteration result"));
+						AddErrorIfFalse(CollectionContains(Collection2, Map2), TEXT("Unexpected InnerJoin iteration result"));
+						AddErrorIfFalse(Map0 != nullptr && Map1 != nullptr && Map0->GetValueType() == Map1->GetValueType(), TEXT("Unexpected InnerJoin iteration result"));
+						AddErrorIfFalse(Map0 != nullptr && Map2 != nullptr && Map0->GetValueType() == Map2->GetValueType(), TEXT("Unexpected InnerJoin iteration result"));
+					},
+					Collection1.CreateMapIterator(),
+					Collection2.CreateMapIterator(),
+					Collection0.CreateMapIterator());
+				AddErrorIfFalse(NumPredicateCalls == 1, TEXT("Unexpected InnerJoin iteration result"));
+
+				NumPredicateCalls = 0;
+				InnerJoinBy(
+					FAttributeCollectionMapJoinOp(),
+					[this, &NumPredicateCalls, &Collection0, &Collection1, &Collection2](FAttributeMap* Map2, FAttributeMap* Map0, FAttributeMap* Map1)
+					{
+						NumPredicateCalls++;
+						AddErrorIfFalse(CollectionContains(Collection0, Map0), TEXT("Unexpected InnerJoin iteration result"));
+						AddErrorIfFalse(CollectionContains(Collection1, Map1), TEXT("Unexpected InnerJoin iteration result"));
+						AddErrorIfFalse(CollectionContains(Collection2, Map2), TEXT("Unexpected InnerJoin iteration result"));
+						AddErrorIfFalse(Map0 != nullptr && Map1 != nullptr && Map0->GetValueType() == Map1->GetValueType(), TEXT("Unexpected InnerJoin iteration result"));
+						AddErrorIfFalse(Map0 != nullptr && Map2 != nullptr && Map0->GetValueType() == Map2->GetValueType(), TEXT("Unexpected InnerJoin iteration result"));
+					},
+					Collection2.CreateMapIterator(),
+					Collection0.CreateMapIterator(),
+					Collection1.CreateMapIterator());
+				AddErrorIfFalse(NumPredicateCalls == 1, TEXT("Unexpected InnerJoin iteration result"));
+
+				NumPredicateCalls = 0;
+				InnerJoinBy(
+					FAttributeCollectionMapJoinOp(),
+					[this, &NumPredicateCalls, &Collection0, &Collection1, &Collection2](FAttributeMap* Map2, FAttributeMap* Map1, FAttributeMap* Map0)
+					{
+						NumPredicateCalls++;
+						AddErrorIfFalse(CollectionContains(Collection0, Map0), TEXT("Unexpected InnerJoin iteration result"));
+						AddErrorIfFalse(CollectionContains(Collection1, Map1), TEXT("Unexpected InnerJoin iteration result"));
+						AddErrorIfFalse(CollectionContains(Collection2, Map2), TEXT("Unexpected InnerJoin iteration result"));
+						AddErrorIfFalse(Map0 != nullptr && Map1 != nullptr && Map0->GetValueType() == Map1->GetValueType(), TEXT("Unexpected InnerJoin iteration result"));
+						AddErrorIfFalse(Map0 != nullptr && Map2 != nullptr && Map0->GetValueType() == Map2->GetValueType(), TEXT("Unexpected InnerJoin iteration result"));
+					},
+					Collection2.CreateMapIterator(),
+					Collection1.CreateMapIterator(),
+					Collection0.CreateMapIterator());
+				AddErrorIfFalse(NumPredicateCalls == 1, TEXT("Unexpected InnerJoin iteration result"));
+			}
+
+			// Join with a transformer list (set maps)
+			{
+				// Use a float transformer
+				// We set up Collection0 such that it has two mappings with float values: Transform -> Float and Float -> Float
+				// InnerJoin should execute the predicate twice, once for each value type that matches the transformer's value type
+
+				// Dummy float sanitize transformer
+				struct FFloatAttributeTransformer_SanitizeTest
+				{
+					using TransformerType = Transformers::FSanitize;
+					using ValueType = FFloatAnimationAttribute;
+
+					static void TransformSetMap(FAttributeSetMap* SetMap) {}	// No-op
+				};
+
+				// Dummy vector sanitize transformer
+				struct FVectorAttributeTransformer_SanitizeTest
+				{
+					using TransformerType = Transformers::FSanitize;
+					using ValueType = FVectorAnimationAttribute;
+
+					static void TransformSetMap(FAttributeSetMap* SetMap) {}	// No-op
+				};
+
+#ifdef _MSC_VER
+	#pragma warning(push)
+	#pragma warning(disable : 4191) // unsafe conversion from 'function pointer type A' to 'function pointer type B'
+#endif
+
+				FAttributeValueTransformerList TransformerList(Transformers::FSanitize::TransformerName);
+				TransformerList.AddSetMapTransformer(FFloatAnimationAttribute::StaticStruct(), reinterpret_cast<FRawTransformerFunc>(&FFloatAttributeTransformer_SanitizeTest::TransformSetMap));
+				TransformerList.AddSetMapTransformer(FVectorAnimationAttribute::StaticStruct(), reinterpret_cast<FRawTransformerFunc>(&FVectorAttributeTransformer_SanitizeTest::TransformSetMap));
+
+#ifdef _MSC_VER
+	#pragma warning(pop)
+#endif
+
+				int32 NumPredicateCalls = 0;
+				InnerJoinBy(
+					FAttributeValueTransformerListWithAttributeCollectionSetMapJoinOp(),
+					[this, &NumPredicateCalls, &Collection0](FRawTransformerFunc Transformer, FAttributeSetMap* SetMap)
+					{
+						NumPredicateCalls++;
+						AddErrorIfFalse(CollectionContains(Collection0, SetMap), TEXT("Unexpected InnerJoin iteration result"));
+						AddErrorIfFalse(SetMap->GetValueType() == FFloatAnimationAttribute::StaticStruct(), TEXT("Unexpected InnerJoin iteration result"));
+					},
+					TransformerList.CreateSetMapTransformerIterator(),
+					Collection0.CreateSetMapIterator());
+				AddErrorIfFalse(NumPredicateCalls == 2, TEXT("Unexpected InnerJoin iteration result"));
+			}
+
+			// Join with a transformer list (name/value maps)
+			{
+				// Use a float transformer
+				// We set up Collection0 such that it has two mappings with these value types: Float and Bone Transform
+				// InnerJoin should execute the predicate twice, once for each value type that matches the transformer's value type
+
+				// Dummy float sanitize transformer
+				struct FFloatAttributeTransformer_SanitizeTest
+				{
+					using TransformerType = Transformers::FSanitize;
+					using ValueType = FFloatAnimationAttribute;
+
+					static void TransformMap(FAttributeMap* Map) {}	// No-op
+				};
+
+				// Dummy vector sanitize transformer
+				struct FVectorAttributeTransformer_SanitizeTest
+				{
+					using TransformerType = Transformers::FSanitize;
+					using ValueType = FVectorAnimationAttribute;
+
+					static void TransformMap(FAttributeMap* Map) {}	// No-op
+				};
+
+#ifdef _MSC_VER
+	#pragma warning(push)
+	#pragma warning(disable : 4191) // unsafe conversion from 'function pointer type A' to 'function pointer type B'
+#endif
+
+				FAttributeValueTransformerList TransformerList(Transformers::FSanitize::TransformerName);
+				TransformerList.AddMapTransformer(FFloatAnimationAttribute::StaticStruct(), reinterpret_cast<FRawTransformerFunc>(&FFloatAttributeTransformer_SanitizeTest::TransformMap));
+				TransformerList.AddMapTransformer(FVectorAnimationAttribute::StaticStruct(), reinterpret_cast<FRawTransformerFunc>(&FVectorAttributeTransformer_SanitizeTest::TransformMap));
+
+#ifdef _MSC_VER
+	#pragma warning(pop)
+#endif
+
+				int32 NumPredicateCalls = 0;
+				InnerJoinBy(
+					FAttributeValueTransformerListWithAttributeCollectionMapJoinOp(),
+					[this, &NumPredicateCalls, &Collection0](FRawTransformerFunc Transformer, FAttributeMap* Map)
+					{
+						NumPredicateCalls++;
+						AddErrorIfFalse(CollectionContains(Collection0, Map), TEXT("Unexpected InnerJoin iteration result"));
+						AddErrorIfFalse(Map->GetValueType() == FFloatAnimationAttribute::StaticStruct(), TEXT("Unexpected InnerJoin iteration result"));
+					},
+					TransformerList.CreateMapTransformerIterator(),
+					Collection0.CreateMapIterator());
+				AddErrorIfFalse(NumPredicateCalls == 1, TEXT("Unexpected InnerJoin iteration result"));
+			}
+		}
+
+		return true;
+	}
+
+	IMPLEMENT_SIMPLE_AUTOMATION_TEST(FIteratorUtilsOuterJoinTest, "Animation.UAF.Attributes.IteratorUtils.OuterJoin", EAutomationTestFlags::EditorContext | EAutomationTestFlags::EngineFilter)
+
+	bool FIteratorUtilsOuterJoinTest::RunTest(const FString& InParameters)
+	{
+		ON_SCOPE_EXIT{ FUtils::CleanupAfterTests(); };
+
+		FAttributeBindingDataPtr BindingData;
+		{
+			UFactory* BindingFactory = NewObject<UFactory>(GetTransientPackage(), UAbstractSkeletonSetBindingFactory::StaticClass());
+			UE_RETURN_ON_ERROR(BindingFactory != nullptr, "Failed to create factory");
+
+			UAbstractSkeletonSetBinding* Binding = CastChecked<UAbstractSkeletonSetBinding>(BindingFactory->FactoryCreateNew(UAbstractSkeletonSetBinding::StaticClass(), GetTransientPackage(), TEXT("TestBindingAsset"), RF_Transient, nullptr, nullptr, NAME_None));
+			UE_RETURN_ON_ERROR(Binding != nullptr, "Failed to create binding asset");
+
+			USkeleton* Skeleton = NewObject<USkeleton>(GetTransientPackage(), USkeleton::StaticClass(), TEXT("TestSkeletonAsset"), RF_Transient);
+			UE_RETURN_ON_ERROR(Skeleton != nullptr, "Failed to create skeleton asset");
+
+			{
+				// Not terribly clean, we cast away the 'const' to modify the skeleton
+				FReferenceSkeleton& RefSkeleton = const_cast<FReferenceSkeleton&>(Skeleton->GetReferenceSkeleton());
+				FReferenceSkeletonModifier SkeletonModifier(RefSkeleton, Skeleton);
+
+				SkeletonModifier.Add(FMeshBoneInfo(TEXT("Root"), TEXT("Root"), INDEX_NONE), FTransform::Identity);
+				SkeletonModifier.Add(FMeshBoneInfo(TEXT("Pelvis"), TEXT("Root"), 0), FTransform::Identity);
+				SkeletonModifier.Add(FMeshBoneInfo(TEXT("Spine_01"), TEXT("Spine_01"), 1), FTransform::Identity);
+				SkeletonModifier.Add(FMeshBoneInfo(TEXT("Spine_02"), TEXT("Spine_02"), 2), FTransform::Identity);
+				SkeletonModifier.Add(FMeshBoneInfo(TEXT("Head"), TEXT("Head"), 3), FTransform::Identity);
+				SkeletonModifier.Add(FMeshBoneInfo(TEXT("LeftArm"), TEXT("LeftArm"), 3), FTransform::Identity);
+				SkeletonModifier.Add(FMeshBoneInfo(TEXT("RightArm"), TEXT("RightArm"), 3), FTransform::Identity);
+
+				// When our modifier is destroyed here, it will rebuild the skeleton
+			}
+
+			UE_RETURN_ON_ERROR(Skeleton->GetReferenceSkeleton().GetNum() == 7, "Failed to build skeleton asset");
+
+			bool bSuccess = Binding->SetSkeleton(Skeleton);
+			UE_RETURN_ON_ERROR(bSuccess, "Failed to set skeleton on our binding");
+
+			UFactory* SecCollectionFactory = NewObject<UFactory>(GetTransientPackage(), UAbstractSkeletonSetCollectionFactory::StaticClass());
+			UE_RETURN_ON_ERROR(SecCollectionFactory != nullptr, "Failed to create factory");
+
+			UAbstractSkeletonSetCollection* SetCollection = CastChecked<UAbstractSkeletonSetCollection>(SecCollectionFactory->FactoryCreateNew(UAbstractSkeletonSetCollection::StaticClass(), GetTransientPackage(), TEXT("TestSetCollectionAsset"), RF_Transient, nullptr, nullptr, NAME_None));
+			UE_RETURN_ON_ERROR(SetCollection != nullptr, "Failed to create set collection asset");
+
+			SetCollection->AddSet(TEXT("FullBody"), NAME_None);
+			SetCollection->AddSet(TEXT("UpperBody"), TEXT("FullBody"));
+
+			bSuccess = Binding->SetSetCollection(SetCollection);
+			UE_RETURN_ON_ERROR(bSuccess, "Failed to set set collection on our binding");
+
+			Binding->AddBoneToSet(TEXT("Root"), TEXT("FullBody"));
+			Binding->AddBoneToSet(TEXT("Pelvis"), TEXT("FullBody"));
+			Binding->AddBoneToSet(TEXT("Spine_01"), TEXT("FullBody"));
+			Binding->AddBoneToSet(TEXT("Spine_02"), TEXT("FullBody"));
+			Binding->AddAttributeToSet(FAnimationAttributeIdentifier(TEXT("AttrFloat0"), INDEX_NONE, NAME_None, FFloatAnimationAttribute::StaticStruct()), TEXT("FullBody"));
+
+			Binding->AddBoneToSet(TEXT("Head"), TEXT("UpperBody"));
+			Binding->AddBoneToSet(TEXT("LeftArm"), TEXT("UpperBody"));
+			Binding->AddBoneToSet(TEXT("RightArm"), TEXT("UpperBody"));
+			Binding->AddAttributeToSet(FAnimationAttributeIdentifier(TEXT("AttrFloat1"), INDEX_NONE, NAME_None, FFloatAnimationAttribute::StaticStruct()), TEXT("UpperBody"));
+
+			BindingData = MakeAttributeBindingData(*Binding);
+			UE_RETURN_ON_ERROR(BindingData.IsValid(), TEXT("Attribute binding data should be valid"));
+		}
+
+		{
+			FAttributeNamedSetPtr FullBodySet = BindingData->FindNamedSet(TEXT("FullBody"));
+			UE_RETURN_ON_ERROR(FullBodySet.IsValid(), TEXT("FullBody named set should be valid"));
+
+			FAttributeCollectionHeap Collection0(FullBodySet);
+
+			AddErrorIfFalse(Collection0.IsEmpty(), TEXT("Attribute collection should be empty"));
+			AddErrorIfFalse(Collection0.GetNamedSet() == FullBodySet, TEXT("Unexpected named set"));
+
+			FAttributeMappingKey FloatToFloatKey = FAttributeMappingKey::MakeFromTo<FFloatAnimationAttribute, FFloatAnimationAttribute>();
+			TAttributeSetMap<FFloatAnimationAttribute>* FloatToFloatMap0 = Collection0.AddSetMap<FFloatAnimationAttribute>(FloatToFloatKey);
+			UE_RETURN_ON_ERROR(FloatToFloatMap0 != nullptr, TEXT("Float->Float map should be valid"));
+
+			FAttributeMappingKey FloatToIntegerKey = FAttributeMappingKey::MakeFromTo<FFloatAnimationAttribute, FIntegerAnimationAttribute>();
+			TAttributeSetMap<FIntegerAnimationAttribute>* FloatToIntegerMap0 = Collection0.AddSetMap<FIntegerAnimationAttribute>(FloatToIntegerKey);
+			UE_RETURN_ON_ERROR(FloatToIntegerMap0 != nullptr, TEXT("Float->Integer map should be valid"));
+
+			FAttributeMappingKey FloatToStringKey = FAttributeMappingKey::MakeFromTo<FFloatAnimationAttribute, FStringAnimationAttribute>();
+			TAttributeSetMap<FStringAnimationAttribute>* FloatToStringMap0 = Collection0.AddSetMap<FStringAnimationAttribute>(FloatToStringKey);
+			UE_RETURN_ON_ERROR(FloatToStringMap0 != nullptr, TEXT("Float->String map should be valid"));
+
+			FAttributeMappingKey TransformToTransformKey = FAttributeMappingKey::MakeFromTo<FBoneTransformAnimationAttribute, FBoneTransformAnimationAttribute>();
+			TAttributeSetMap<FBoneTransformAnimationAttribute>* TransformToTransformMap0 = Collection0.AddSetMap<FBoneTransformAnimationAttribute>(TransformToTransformKey);
+			UE_RETURN_ON_ERROR(TransformToTransformMap0 != nullptr, TEXT("Transform->Transform map should be valid"));
+
+			FAttributeMappingKey TransformToFloatKey = FAttributeMappingKey::MakeFromTo<FBoneTransformAnimationAttribute, FFloatAnimationAttribute>();
+			TAttributeSetMap<FFloatAnimationAttribute>* TransformToFloatMap0 = Collection0.AddSetMap<FFloatAnimationAttribute>(TransformToFloatKey);
+			UE_RETURN_ON_ERROR(TransformToFloatMap0 != nullptr, TEXT("Transform->Bone map should be valid"));
+
+			TAttributeMap<FFloatAnimationAttribute>* FloatMap0 = Collection0.AddMap<FFloatAnimationAttribute>();
+			UE_RETURN_ON_ERROR(FloatMap0 != nullptr, TEXT("Float map should be valid"));
+
+			TAttributeMap<FBoneTransformAnimationAttribute>* BoneMap0 = Collection0.AddMap<FBoneTransformAnimationAttribute>();
+			UE_RETURN_ON_ERROR(BoneMap0 != nullptr, TEXT("Bone map should be valid"));
+
+			AddErrorIfFalse(!Collection0.IsEmpty(), TEXT("Attribute collection should not be empty"));
+			AddErrorIfFalse(Collection0.NumSetMaps() == 5, TEXT("Unexpected set map count"));
+			AddErrorIfFalse(Collection0.NumMaps() == 2, TEXT("Unexpected map count"));
+
+			FAttributeCollectionHeap Collection1(FullBodySet);
+
+			AddErrorIfFalse(Collection1.IsEmpty(), TEXT("Attribute collection should be empty"));
+			AddErrorIfFalse(Collection1.GetNamedSet() == FullBodySet, TEXT("Unexpected named set"));
+
+			TAttributeSetMap<FFloatAnimationAttribute>* FloatToFloatMap1 = Collection1.AddSetMap<FFloatAnimationAttribute>(FloatToFloatKey);
+			UE_RETURN_ON_ERROR(FloatToFloatMap1 != nullptr, TEXT("Float->Float map should be valid"));
+
+			TAttributeSetMap<FIntegerAnimationAttribute>* FloatToIntegerMap1 = Collection1.AddSetMap<FIntegerAnimationAttribute>(FloatToIntegerKey);
+			UE_RETURN_ON_ERROR(FloatToIntegerMap1 != nullptr, TEXT("Float->Integer map should be valid"));
+
+			TAttributeSetMap<FStringAnimationAttribute>* FloatToStringMap1 = Collection1.AddSetMap<FStringAnimationAttribute>(FloatToStringKey);
+			UE_RETURN_ON_ERROR(FloatToStringMap1 != nullptr, TEXT("Float->String map should be valid"));
+
+			TAttributeSetMap<FBoneTransformAnimationAttribute>* TransformToTransformMap1 = Collection1.AddSetMap<FBoneTransformAnimationAttribute>(TransformToTransformKey);
+			UE_RETURN_ON_ERROR(TransformToTransformMap1 != nullptr, TEXT("Transform->Transform map should be valid"));
+
+			TAttributeSetMap<FFloatAnimationAttribute>* TransformToFloatMap1 = Collection1.AddSetMap<FFloatAnimationAttribute>(TransformToFloatKey);
+			UE_RETURN_ON_ERROR(TransformToFloatMap1 != nullptr, TEXT("Transform->Bone map should be valid"));
+
+			TAttributeMap<FFloatAnimationAttribute>* FloatMap1 = Collection1.AddMap<FFloatAnimationAttribute>();
+			UE_RETURN_ON_ERROR(FloatMap1 != nullptr, TEXT("Float map should be valid"));
+
+			TAttributeMap<FBoneTransformAnimationAttribute>* BoneMap1 = Collection1.AddMap<FBoneTransformAnimationAttribute>();
+			UE_RETURN_ON_ERROR(BoneMap1 != nullptr, TEXT("Bone map should be valid"));
+
+			AddErrorIfFalse(!Collection1.IsEmpty(), TEXT("Attribute collection should not be empty"));
+			AddErrorIfFalse(Collection1.NumSetMaps() == 5, TEXT("Unexpected set map count"));
+			AddErrorIfFalse(Collection1.NumMaps() == 2, TEXT("Unexpected map count"));
+
+			// Both inputs perfectly match (set maps)
+			{
+				int32 NumPredicateCalls = 0;
+				OuterJoinBy(
+					FAttributeCollectionSetMapJoinOp(),
+					[this, &NumPredicateCalls, &Collection0, &Collection1](FAttributeSetMap* Map0, FAttributeSetMap* Map1)
+					{
+						NumPredicateCalls++;
+						AddErrorIfFalse(CollectionContains(Collection0, Map0), TEXT("Unexpected OuterJoin iteration result"));
+						AddErrorIfFalse(CollectionContains(Collection1, Map1), TEXT("Unexpected OuterJoin iteration result"));
+						AddErrorIfFalse(Map0 != nullptr && Map1 != nullptr && Map0->GetAttributeType() == Map1->GetAttributeType(), TEXT("Unexpected OuterJoin iteration result"));
+						AddErrorIfFalse(Map0 != nullptr && Map1 != nullptr && Map0->GetValueType() == Map1->GetValueType(), TEXT("Unexpected OuterJoin iteration result"));
+					},
+					Collection0.CreateSetMapIterator(),
+					Collection1.CreateSetMapIterator());
+				AddErrorIfFalse(NumPredicateCalls == 5, TEXT("Unexpected OuterJoin iteration result"));
+
+				NumPredicateCalls = 0;
+				OuterJoinBy(
+					FAttributeCollectionSetMapJoinOp(),
+					[this, &NumPredicateCalls, &Collection0, &Collection1](const FAttributeSetMap* Map0, const FAttributeSetMap* Map1)
+					{
+						NumPredicateCalls++;
+						AddErrorIfFalse(CollectionContains(Collection0, Map0), TEXT("Unexpected OuterJoin iteration result"));
+						AddErrorIfFalse(CollectionContains(Collection1, Map1), TEXT("Unexpected OuterJoin iteration result"));
+						AddErrorIfFalse(Map0 != nullptr && Map1 != nullptr && Map0->GetAttributeType() == Map1->GetAttributeType(), TEXT("Unexpected OuterJoin iteration result"));
+						AddErrorIfFalse(Map0 != nullptr && Map1 != nullptr && Map0->GetValueType() == Map1->GetValueType(), TEXT("Unexpected OuterJoin iteration result"));
+					},
+					Collection0.CreateConstSetMapIterator(),
+					Collection1.CreateConstSetMapIterator());
+				AddErrorIfFalse(NumPredicateCalls == 5, TEXT("Unexpected OuterJoin iteration result"));
+
+				NumPredicateCalls = 0;
+				OuterJoinBy(
+					FAttributeCollectionSetMapJoinOp(),
+					[this, &NumPredicateCalls, &Collection0, &Collection1](const FAttributeSetMap* Map0, FAttributeSetMap* Map1)
+					{
+						NumPredicateCalls++;
+						AddErrorIfFalse(CollectionContains(Collection0, Map0), TEXT("Unexpected OuterJoin iteration result"));
+						AddErrorIfFalse(CollectionContains(Collection1, Map1), TEXT("Unexpected OuterJoin iteration result"));
+						AddErrorIfFalse(Map0 != nullptr && Map1 != nullptr && Map0->GetAttributeType() == Map1->GetAttributeType(), TEXT("Unexpected OuterJoin iteration result"));
+						AddErrorIfFalse(Map0 != nullptr && Map1 != nullptr && Map0->GetValueType() == Map1->GetValueType(), TEXT("Unexpected OuterJoin iteration result"));
+					},
+					Collection0.CreateConstSetMapIterator(),
+					Collection1.CreateSetMapIterator());
+				AddErrorIfFalse(NumPredicateCalls == 5, TEXT("Unexpected OuterJoin iteration result"));
+
+				NumPredicateCalls = 0;
+				OuterJoinBy(
+					FAttributeCollectionSetMapJoinOp(),
+					[this, &NumPredicateCalls, &Collection0, &Collection1](FAttributeSetMap* Map0, const FAttributeSetMap* Map1)
+					{
+						NumPredicateCalls++;
+						AddErrorIfFalse(CollectionContains(Collection0, Map0), TEXT("Unexpected OuterJoin iteration result"));
+						AddErrorIfFalse(CollectionContains(Collection1, Map1), TEXT("Unexpected OuterJoin iteration result"));
+						AddErrorIfFalse(Map0 != nullptr && Map1 != nullptr && Map0->GetAttributeType() == Map1->GetAttributeType(), TEXT("Unexpected OuterJoin iteration result"));
+						AddErrorIfFalse(Map0 != nullptr && Map1 != nullptr && Map0->GetValueType() == Map1->GetValueType(), TEXT("Unexpected OuterJoin iteration result"));
+					},
+					Collection0.CreateSetMapIterator(),
+					Collection1.CreateConstSetMapIterator());
+				AddErrorIfFalse(NumPredicateCalls == 5, TEXT("Unexpected OuterJoin iteration result"));
+			}
+
+			// Both inputs perfectly match (name/value maps)
+			{
+				int32 NumPredicateCalls = 0;
+				OuterJoinBy(
+					FAttributeCollectionMapJoinOp(),
+					[this, &NumPredicateCalls, &Collection0, &Collection1](FAttributeMap* Map0, FAttributeMap* Map1)
+					{
+						NumPredicateCalls++;
+						AddErrorIfFalse(CollectionContains(Collection0, Map0), TEXT("Unexpected OuterJoin iteration result"));
+						AddErrorIfFalse(CollectionContains(Collection1, Map1), TEXT("Unexpected OuterJoin iteration result"));
+						AddErrorIfFalse(Map0 != nullptr && Map1 != nullptr && Map0->GetValueType() == Map1->GetValueType(), TEXT("Unexpected OuterJoin iteration result"));
+					},
+					Collection0.CreateMapIterator(),
+					Collection1.CreateMapIterator());
+				AddErrorIfFalse(NumPredicateCalls == 2, TEXT("Unexpected OuterJoin iteration result"));
+
+				NumPredicateCalls = 0;
+				OuterJoinBy(
+					FAttributeCollectionMapJoinOp(),
+					[this, &NumPredicateCalls, &Collection0, &Collection1](const FAttributeMap* Map0, const FAttributeMap* Map1)
+					{
+						NumPredicateCalls++;
+						AddErrorIfFalse(CollectionContains(Collection0, Map0), TEXT("Unexpected OuterJoin iteration result"));
+						AddErrorIfFalse(CollectionContains(Collection1, Map1), TEXT("Unexpected OuterJoin iteration result"));
+						AddErrorIfFalse(Map0 != nullptr && Map1 != nullptr && Map0->GetValueType() == Map1->GetValueType(), TEXT("Unexpected OuterJoin iteration result"));
+					},
+					Collection0.CreateConstMapIterator(),
+					Collection1.CreateConstMapIterator());
+				AddErrorIfFalse(NumPredicateCalls == 2, TEXT("Unexpected OuterJoin iteration result"));
+
+				NumPredicateCalls = 0;
+				OuterJoinBy(
+					FAttributeCollectionMapJoinOp(),
+					[this, &NumPredicateCalls, &Collection0, &Collection1](const FAttributeMap* Map0, FAttributeMap* Map1)
+					{
+						NumPredicateCalls++;
+						AddErrorIfFalse(CollectionContains(Collection0, Map0), TEXT("Unexpected OuterJoin iteration result"));
+						AddErrorIfFalse(CollectionContains(Collection1, Map1), TEXT("Unexpected OuterJoin iteration result"));
+						AddErrorIfFalse(Map0 != nullptr && Map1 != nullptr && Map0->GetValueType() == Map1->GetValueType(), TEXT("Unexpected OuterJoin iteration result"));
+					},
+					Collection0.CreateConstMapIterator(),
+					Collection1.CreateMapIterator());
+				AddErrorIfFalse(NumPredicateCalls == 2, TEXT("Unexpected OuterJoin iteration result"));
+
+				NumPredicateCalls = 0;
+				OuterJoinBy(
+					FAttributeCollectionMapJoinOp(),
+					[this, &NumPredicateCalls, &Collection0, &Collection1](FAttributeMap* Map0, const FAttributeMap* Map1)
+					{
+						NumPredicateCalls++;
+						AddErrorIfFalse(CollectionContains(Collection0, Map0), TEXT("Unexpected OuterJoin iteration result"));
+						AddErrorIfFalse(CollectionContains(Collection1, Map1), TEXT("Unexpected OuterJoin iteration result"));
+						AddErrorIfFalse(Map0 != nullptr && Map1 != nullptr && Map0->GetValueType() == Map1->GetValueType(), TEXT("Unexpected OuterJoin iteration result"));
+					},
+					Collection0.CreateMapIterator(),
+					Collection1.CreateConstMapIterator());
+				AddErrorIfFalse(NumPredicateCalls == 2, TEXT("Unexpected OuterJoin iteration result"));
+			}
+
+			FAttributeCollectionHeap Collection2(FullBodySet);
+
+			AddErrorIfFalse(Collection2.IsEmpty(), TEXT("Attribute collection should be empty"));
+			AddErrorIfFalse(Collection2.GetNamedSet() == FullBodySet, TEXT("Unexpected named set"));
+
+			TAttributeSetMap<FFloatAnimationAttribute>* FloatToFloatMap2 = Collection2.AddSetMap<FFloatAnimationAttribute>(FloatToFloatKey);
+			UE_RETURN_ON_ERROR(FloatToFloatMap2 != nullptr, TEXT("Float->Float map should be valid"));
+
+			TAttributeSetMap<FBoneTransformAnimationAttribute>* TransformToTransformMap2 = Collection2.AddSetMap<FBoneTransformAnimationAttribute>(TransformToTransformKey);
+			UE_RETURN_ON_ERROR(TransformToTransformMap2 != nullptr, TEXT("Transform->Transform map should be valid"));
+
+			TAttributeMap<FFloatAnimationAttribute>* FloatMap2 = Collection2.AddMap<FFloatAnimationAttribute>();
+			UE_RETURN_ON_ERROR(FloatMap2 != nullptr, TEXT("Float map should be valid"));
+
+			AddErrorIfFalse(!Collection2.IsEmpty(), TEXT("Attribute collection should not be empty"));
+			AddErrorIfFalse(Collection2.NumSetMaps() == 2, TEXT("Unexpected set map count"));
+			AddErrorIfFalse(Collection2.NumMaps() == 1, TEXT("Unexpected map count"));
+
+			// First input is shorter and will be padded with null entries (set maps)
+			{
+				int32 NumPredicateCalls = 0;
+				int32 NumNullEntries = 0;
+				OuterJoinBy(
+					FAttributeCollectionSetMapJoinOp(),
+					[this, &NumPredicateCalls, &NumNullEntries, &Collection0, &Collection2](FAttributeSetMap* Map2, FAttributeSetMap* Map0)
+					{
+						NumPredicateCalls++;
+						AddErrorIfFalse(CollectionContains(Collection0, Map0), TEXT("Unexpected OuterJoin iteration result"));
+
+						if (Map2 != nullptr)
+						{
+							AddErrorIfFalse(CollectionContains(Collection2, Map2), TEXT("Unexpected OuterJoin iteration result"));
+							AddErrorIfFalse(Map0 != nullptr && Map0->GetAttributeType() == Map2->GetAttributeType(), TEXT("Unexpected OuterJoin iteration result"));
+							AddErrorIfFalse(Map0 != nullptr && Map0->GetValueType() == Map2->GetValueType(), TEXT("Unexpected OuterJoin iteration result"));
+						}
+						else
+						{
+							NumNullEntries++;
+							AddErrorIfFalse(Map0 != nullptr, TEXT("Unexpected OuterJoin iteration result"));
+						}
+					},
+					Collection2.CreateSetMapIterator(),
+					Collection0.CreateSetMapIterator());
+				AddErrorIfFalse(NumPredicateCalls == 5, TEXT("Unexpected OuterJoin iteration result"));
+				AddErrorIfFalse(NumNullEntries == 3, TEXT("Unexpected OuterJoin iteration result"));
+
+				NumPredicateCalls = 0;
+				NumNullEntries = 0;
+				OuterJoinBy(
+					FAttributeCollectionSetMapJoinOp(),
+					[this, &NumPredicateCalls, &NumNullEntries, &Collection0, &Collection2](const FAttributeSetMap* Map2, const FAttributeSetMap* Map0)
+					{
+						NumPredicateCalls++;
+						AddErrorIfFalse(CollectionContains(Collection0, Map0), TEXT("Unexpected OuterJoin iteration result"));
+
+						if (Map2 != nullptr)
+						{
+							AddErrorIfFalse(CollectionContains(Collection2, Map2), TEXT("Unexpected OuterJoin iteration result"));
+							AddErrorIfFalse(Map0 != nullptr && Map0->GetAttributeType() == Map2->GetAttributeType(), TEXT("Unexpected OuterJoin iteration result"));
+							AddErrorIfFalse(Map0 != nullptr && Map0->GetValueType() == Map2->GetValueType(), TEXT("Unexpected OuterJoin iteration result"));
+						}
+						else
+						{
+							NumNullEntries++;
+							AddErrorIfFalse(Map0 != nullptr, TEXT("Unexpected OuterJoin iteration result"));
+						}
+					},
+					Collection2.CreateConstSetMapIterator(),
+					Collection0.CreateConstSetMapIterator());
+				AddErrorIfFalse(NumPredicateCalls == 5, TEXT("Unexpected OuterJoin iteration result"));
+				AddErrorIfFalse(NumNullEntries == 3, TEXT("Unexpected OuterJoin iteration result"));
+
+				NumPredicateCalls = 0;
+				NumNullEntries = 0;
+				OuterJoinBy(
+					FAttributeCollectionSetMapJoinOp(),
+					[this, &NumPredicateCalls, &NumNullEntries, &Collection0, &Collection2](const FAttributeSetMap* Map2, FAttributeSetMap* Map0)
+					{
+						NumPredicateCalls++;
+						AddErrorIfFalse(CollectionContains(Collection0, Map0), TEXT("Unexpected OuterJoin iteration result"));
+
+						if (Map2 != nullptr)
+						{
+							AddErrorIfFalse(CollectionContains(Collection2, Map2), TEXT("Unexpected OuterJoin iteration result"));
+							AddErrorIfFalse(Map0 != nullptr && Map0->GetAttributeType() == Map2->GetAttributeType(), TEXT("Unexpected OuterJoin iteration result"));
+							AddErrorIfFalse(Map0 != nullptr && Map0->GetValueType() == Map2->GetValueType(), TEXT("Unexpected OuterJoin iteration result"));
+						}
+						else
+						{
+							NumNullEntries++;
+							AddErrorIfFalse(Map0 != nullptr, TEXT("Unexpected OuterJoin iteration result"));
+						}
+					},
+					Collection2.CreateConstSetMapIterator(),
+					Collection0.CreateSetMapIterator());
+				AddErrorIfFalse(NumPredicateCalls == 5, TEXT("Unexpected OuterJoin iteration result"));
+				AddErrorIfFalse(NumNullEntries == 3, TEXT("Unexpected OuterJoin iteration result"));
+
+				NumPredicateCalls = 0;
+				NumNullEntries = 0;
+				OuterJoinBy(
+					FAttributeCollectionSetMapJoinOp(),
+					[this, &NumPredicateCalls, &NumNullEntries, &Collection0, &Collection2](FAttributeSetMap* Map2, const FAttributeSetMap* Map0)
+					{
+						NumPredicateCalls++;
+						AddErrorIfFalse(CollectionContains(Collection0, Map0), TEXT("Unexpected OuterJoin iteration result"));
+
+						if (Map2 != nullptr)
+						{
+							AddErrorIfFalse(CollectionContains(Collection2, Map2), TEXT("Unexpected OuterJoin iteration result"));
+							AddErrorIfFalse(Map0 != nullptr && Map0->GetAttributeType() == Map2->GetAttributeType(), TEXT("Unexpected OuterJoin iteration result"));
+							AddErrorIfFalse(Map0 != nullptr && Map0->GetValueType() == Map2->GetValueType(), TEXT("Unexpected OuterJoin iteration result"));
+						}
+						else
+						{
+							NumNullEntries++;
+							AddErrorIfFalse(Map0 != nullptr, TEXT("Unexpected OuterJoin iteration result"));
+						}
+					},
+					Collection2.CreateSetMapIterator(),
+					Collection0.CreateConstSetMapIterator());
+				AddErrorIfFalse(NumPredicateCalls == 5, TEXT("Unexpected OuterJoin iteration result"));
+				AddErrorIfFalse(NumNullEntries == 3, TEXT("Unexpected OuterJoin iteration result"));
+			}
+
+			// First input is shorter and will be padded with null entries (name/value maps)
+			{
+				int32 NumPredicateCalls = 0;
+				int32 NumNullEntries = 0;
+				OuterJoinBy(
+					FAttributeCollectionMapJoinOp(),
+					[this, &NumPredicateCalls, &NumNullEntries, &Collection0, &Collection2](FAttributeMap* Map2, FAttributeMap* Map0)
+					{
+						NumPredicateCalls++;
+						AddErrorIfFalse(CollectionContains(Collection0, Map0), TEXT("Unexpected OuterJoin iteration result"));
+
+						if (Map2 != nullptr)
+						{
+							AddErrorIfFalse(CollectionContains(Collection2, Map2), TEXT("Unexpected OuterJoin iteration result"));
+							AddErrorIfFalse(Map0 != nullptr && Map0->GetValueType() == Map2->GetValueType(), TEXT("Unexpected OuterJoin iteration result"));
+						}
+						else
+						{
+							NumNullEntries++;
+							AddErrorIfFalse(Map0 != nullptr, TEXT("Unexpected OuterJoin iteration result"));
+						}
+					},
+					Collection2.CreateMapIterator(),
+					Collection0.CreateMapIterator());
+				AddErrorIfFalse(NumPredicateCalls == 2, TEXT("Unexpected OuterJoin iteration result"));
+				AddErrorIfFalse(NumNullEntries == 1, TEXT("Unexpected OuterJoin iteration result"));
+
+				NumPredicateCalls = 0;
+				NumNullEntries = 0;
+				OuterJoinBy(
+					FAttributeCollectionMapJoinOp(),
+					[this, &NumPredicateCalls, &NumNullEntries, &Collection0, &Collection2](const FAttributeMap* Map2, const FAttributeMap* Map0)
+					{
+						NumPredicateCalls++;
+						AddErrorIfFalse(CollectionContains(Collection0, Map0), TEXT("Unexpected OuterJoin iteration result"));
+
+						if (Map2 != nullptr)
+						{
+							AddErrorIfFalse(CollectionContains(Collection2, Map2), TEXT("Unexpected OuterJoin iteration result"));
+							AddErrorIfFalse(Map0 != nullptr && Map0->GetValueType() == Map2->GetValueType(), TEXT("Unexpected OuterJoin iteration result"));
+						}
+						else
+						{
+							NumNullEntries++;
+							AddErrorIfFalse(Map0 != nullptr, TEXT("Unexpected OuterJoin iteration result"));
+						}
+					},
+					Collection2.CreateConstMapIterator(),
+					Collection0.CreateConstMapIterator());
+				AddErrorIfFalse(NumPredicateCalls == 2, TEXT("Unexpected OuterJoin iteration result"));
+				AddErrorIfFalse(NumNullEntries == 1, TEXT("Unexpected OuterJoin iteration result"));
+
+				NumPredicateCalls = 0;
+				NumNullEntries = 0;
+				OuterJoinBy(
+					FAttributeCollectionMapJoinOp(),
+					[this, &NumPredicateCalls, &NumNullEntries, &Collection0, &Collection2](const FAttributeMap* Map2, FAttributeMap* Map0)
+					{
+						NumPredicateCalls++;
+						AddErrorIfFalse(CollectionContains(Collection0, Map0), TEXT("Unexpected OuterJoin iteration result"));
+
+						if (Map2 != nullptr)
+						{
+							AddErrorIfFalse(CollectionContains(Collection2, Map2), TEXT("Unexpected OuterJoin iteration result"));
+							AddErrorIfFalse(Map0 != nullptr && Map0->GetValueType() == Map2->GetValueType(), TEXT("Unexpected OuterJoin iteration result"));
+						}
+						else
+						{
+							NumNullEntries++;
+							AddErrorIfFalse(Map0 != nullptr, TEXT("Unexpected OuterJoin iteration result"));
+						}
+					},
+					Collection2.CreateConstMapIterator(),
+					Collection0.CreateMapIterator());
+				AddErrorIfFalse(NumPredicateCalls == 2, TEXT("Unexpected OuterJoin iteration result"));
+				AddErrorIfFalse(NumNullEntries == 1, TEXT("Unexpected OuterJoin iteration result"));
+
+				NumPredicateCalls = 0;
+				NumNullEntries = 0;
+				OuterJoinBy(
+					FAttributeCollectionMapJoinOp(),
+					[this, &NumPredicateCalls, &NumNullEntries, &Collection0, &Collection2](FAttributeMap* Map2, const FAttributeMap* Map0)
+					{
+						NumPredicateCalls++;
+						AddErrorIfFalse(CollectionContains(Collection0, Map0), TEXT("Unexpected OuterJoin iteration result"));
+
+						if (Map2 != nullptr)
+						{
+							AddErrorIfFalse(CollectionContains(Collection2, Map2), TEXT("Unexpected OuterJoin iteration result"));
+							AddErrorIfFalse(Map0 != nullptr && Map0->GetValueType() == Map2->GetValueType(), TEXT("Unexpected OuterJoin iteration result"));
+						}
+						else
+						{
+							NumNullEntries++;
+							AddErrorIfFalse(Map0 != nullptr, TEXT("Unexpected OuterJoin iteration result"));
+						}
+					},
+					Collection2.CreateMapIterator(),
+					Collection0.CreateConstMapIterator());
+				AddErrorIfFalse(NumPredicateCalls == 2, TEXT("Unexpected OuterJoin iteration result"));
+				AddErrorIfFalse(NumNullEntries == 1, TEXT("Unexpected OuterJoin iteration result"));
+			}
+
+			// Second input is shorter and will be padded with null entries (set maps)
+			{
+				int32 NumPredicateCalls = 0;
+				int32 NumNullEntries = 0;
+				OuterJoinBy(
+					FAttributeCollectionSetMapJoinOp(),
+					[this, &NumPredicateCalls, &NumNullEntries, &Collection0, &Collection2](FAttributeSetMap* Map0, FAttributeSetMap* Map2)
+					{
+						NumPredicateCalls++;
+						AddErrorIfFalse(CollectionContains(Collection0, Map0), TEXT("Unexpected OuterJoin iteration result"));
+
+						if (Map2 != nullptr)
+						{
+							AddErrorIfFalse(CollectionContains(Collection2, Map2), TEXT("Unexpected OuterJoin iteration result"));
+							AddErrorIfFalse(Map0 != nullptr && Map0->GetAttributeType() == Map2->GetAttributeType(), TEXT("Unexpected OuterJoin iteration result"));
+							AddErrorIfFalse(Map0 != nullptr && Map0->GetValueType() == Map2->GetValueType(), TEXT("Unexpected OuterJoin iteration result"));
+						}
+						else
+						{
+							NumNullEntries++;
+							AddErrorIfFalse(Map0 != nullptr, TEXT("Unexpected OuterJoin iteration result"));
+						}
+					},
+					Collection0.CreateSetMapIterator(),
+					Collection2.CreateSetMapIterator());
+				AddErrorIfFalse(NumPredicateCalls == 5, TEXT("Unexpected OuterJoin iteration result"));
+				AddErrorIfFalse(NumNullEntries == 3, TEXT("Unexpected OuterJoin iteration result"));
+
+				NumPredicateCalls = 0;
+				NumNullEntries = 0;
+				OuterJoinBy(
+					FAttributeCollectionSetMapJoinOp(),
+					[this, &NumPredicateCalls, &NumNullEntries, &Collection0, &Collection2](const FAttributeSetMap* Map0, const FAttributeSetMap* Map2)
+					{
+						NumPredicateCalls++;
+						AddErrorIfFalse(CollectionContains(Collection0, Map0), TEXT("Unexpected OuterJoin iteration result"));
+
+						if (Map2 != nullptr)
+						{
+							AddErrorIfFalse(CollectionContains(Collection2, Map2), TEXT("Unexpected OuterJoin iteration result"));
+							AddErrorIfFalse(Map0 != nullptr && Map0->GetAttributeType() == Map2->GetAttributeType(), TEXT("Unexpected OuterJoin iteration result"));
+							AddErrorIfFalse(Map0 != nullptr && Map0->GetValueType() == Map2->GetValueType(), TEXT("Unexpected OuterJoin iteration result"));
+						}
+						else
+						{
+							NumNullEntries++;
+							AddErrorIfFalse(Map0 != nullptr, TEXT("Unexpected OuterJoin iteration result"));
+						}
+					},
+					Collection0.CreateConstSetMapIterator(),
+					Collection2.CreateConstSetMapIterator());
+				AddErrorIfFalse(NumPredicateCalls == 5, TEXT("Unexpected OuterJoin iteration result"));
+				AddErrorIfFalse(NumNullEntries == 3, TEXT("Unexpected OuterJoin iteration result"));
+
+				NumPredicateCalls = 0;
+				NumNullEntries = 0;
+				OuterJoinBy(
+					FAttributeCollectionSetMapJoinOp(),
+					[this, &NumPredicateCalls, &NumNullEntries, &Collection0, &Collection2](const FAttributeSetMap* Map0, FAttributeSetMap* Map2)
+					{
+						NumPredicateCalls++;
+						AddErrorIfFalse(CollectionContains(Collection0, Map0), TEXT("Unexpected OuterJoin iteration result"));
+
+						if (Map2 != nullptr)
+						{
+							AddErrorIfFalse(CollectionContains(Collection2, Map2), TEXT("Unexpected OuterJoin iteration result"));
+							AddErrorIfFalse(Map0 != nullptr && Map0->GetAttributeType() == Map2->GetAttributeType(), TEXT("Unexpected OuterJoin iteration result"));
+							AddErrorIfFalse(Map0 != nullptr && Map0->GetValueType() == Map2->GetValueType(), TEXT("Unexpected OuterJoin iteration result"));
+						}
+						else
+						{
+							NumNullEntries++;
+							AddErrorIfFalse(Map0 != nullptr, TEXT("Unexpected OuterJoin iteration result"));
+						}
+					},
+					Collection0.CreateConstSetMapIterator(),
+					Collection2.CreateSetMapIterator());
+				AddErrorIfFalse(NumPredicateCalls == 5, TEXT("Unexpected OuterJoin iteration result"));
+				AddErrorIfFalse(NumNullEntries == 3, TEXT("Unexpected OuterJoin iteration result"));
+
+				NumPredicateCalls = 0;
+				NumNullEntries = 0;
+				OuterJoinBy(
+					FAttributeCollectionSetMapJoinOp(),
+					[this, &NumPredicateCalls, &NumNullEntries, &Collection0, &Collection2](FAttributeSetMap* Map0, const FAttributeSetMap* Map2)
+					{
+						NumPredicateCalls++;
+						AddErrorIfFalse(CollectionContains(Collection0, Map0), TEXT("Unexpected OuterJoin iteration result"));
+
+						if (Map2 != nullptr)
+						{
+							AddErrorIfFalse(CollectionContains(Collection2, Map2), TEXT("Unexpected OuterJoin iteration result"));
+							AddErrorIfFalse(Map0 != nullptr && Map0->GetAttributeType() == Map2->GetAttributeType(), TEXT("Unexpected OuterJoin iteration result"));
+							AddErrorIfFalse(Map0 != nullptr && Map0->GetValueType() == Map2->GetValueType(), TEXT("Unexpected OuterJoin iteration result"));
+						}
+						else
+						{
+							NumNullEntries++;
+							AddErrorIfFalse(Map0 != nullptr, TEXT("Unexpected OuterJoin iteration result"));
+						}
+					},
+					Collection0.CreateSetMapIterator(),
+					Collection2.CreateConstSetMapIterator());
+				AddErrorIfFalse(NumPredicateCalls == 5, TEXT("Unexpected OuterJoin iteration result"));
+				AddErrorIfFalse(NumNullEntries == 3, TEXT("Unexpected OuterJoin iteration result"));
+			}
+
+			// Second input is shorter and will be padded with null entries (name/value maps)
+			{
+				int32 NumPredicateCalls = 0;
+				int32 NumNullEntries = 0;
+				OuterJoinBy(
+					FAttributeCollectionMapJoinOp(),
+					[this, &NumPredicateCalls, &NumNullEntries, &Collection0, &Collection2](FAttributeMap* Map0, FAttributeMap* Map2)
+					{
+						NumPredicateCalls++;
+						AddErrorIfFalse(CollectionContains(Collection0, Map0), TEXT("Unexpected OuterJoin iteration result"));
+
+						if (Map2 != nullptr)
+						{
+							AddErrorIfFalse(CollectionContains(Collection2, Map2), TEXT("Unexpected OuterJoin iteration result"));
+							AddErrorIfFalse(Map0 != nullptr && Map0->GetValueType() == Map2->GetValueType(), TEXT("Unexpected OuterJoin iteration result"));
+						}
+						else
+						{
+							NumNullEntries++;
+							AddErrorIfFalse(Map0 != nullptr, TEXT("Unexpected OuterJoin iteration result"));
+						}
+					},
+					Collection0.CreateMapIterator(),
+					Collection2.CreateMapIterator());
+				AddErrorIfFalse(NumPredicateCalls == 2, TEXT("Unexpected OuterJoin iteration result"));
+				AddErrorIfFalse(NumNullEntries == 1, TEXT("Unexpected OuterJoin iteration result"));
+
+				NumPredicateCalls = 0;
+				NumNullEntries = 0;
+				OuterJoinBy(
+					FAttributeCollectionMapJoinOp(),
+					[this, &NumPredicateCalls, &NumNullEntries, &Collection0, &Collection2](const FAttributeMap* Map0, const FAttributeMap* Map2)
+					{
+						NumPredicateCalls++;
+						AddErrorIfFalse(CollectionContains(Collection0, Map0), TEXT("Unexpected OuterJoin iteration result"));
+
+						if (Map2 != nullptr)
+						{
+							AddErrorIfFalse(CollectionContains(Collection2, Map2), TEXT("Unexpected OuterJoin iteration result"));
+							AddErrorIfFalse(Map0 != nullptr && Map0->GetValueType() == Map2->GetValueType(), TEXT("Unexpected OuterJoin iteration result"));
+						}
+						else
+						{
+							NumNullEntries++;
+							AddErrorIfFalse(Map0 != nullptr, TEXT("Unexpected OuterJoin iteration result"));
+						}
+					},
+					Collection0.CreateConstMapIterator(),
+					Collection2.CreateConstMapIterator());
+				AddErrorIfFalse(NumPredicateCalls == 2, TEXT("Unexpected OuterJoin iteration result"));
+				AddErrorIfFalse(NumNullEntries == 1, TEXT("Unexpected OuterJoin iteration result"));
+
+				NumPredicateCalls = 0;
+				NumNullEntries = 0;
+				OuterJoinBy(
+					FAttributeCollectionMapJoinOp(),
+					[this, &NumPredicateCalls, &NumNullEntries, &Collection0, &Collection2](const FAttributeMap* Map0, FAttributeMap* Map2)
+					{
+						NumPredicateCalls++;
+						AddErrorIfFalse(CollectionContains(Collection0, Map0), TEXT("Unexpected OuterJoin iteration result"));
+
+						if (Map2 != nullptr)
+						{
+							AddErrorIfFalse(CollectionContains(Collection2, Map2), TEXT("Unexpected OuterJoin iteration result"));
+							AddErrorIfFalse(Map0 != nullptr && Map0->GetValueType() == Map2->GetValueType(), TEXT("Unexpected OuterJoin iteration result"));
+						}
+						else
+						{
+							NumNullEntries++;
+							AddErrorIfFalse(Map0 != nullptr, TEXT("Unexpected OuterJoin iteration result"));
+						}
+					},
+					Collection0.CreateConstMapIterator(),
+					Collection2.CreateMapIterator());
+				AddErrorIfFalse(NumPredicateCalls == 2, TEXT("Unexpected OuterJoin iteration result"));
+				AddErrorIfFalse(NumNullEntries == 1, TEXT("Unexpected OuterJoin iteration result"));
+
+				NumPredicateCalls = 0;
+				NumNullEntries = 0;
+				OuterJoinBy(
+					FAttributeCollectionMapJoinOp(),
+					[this, &NumPredicateCalls, &NumNullEntries, &Collection0, &Collection2](FAttributeMap* Map0, const FAttributeMap* Map2)
+					{
+						NumPredicateCalls++;
+						AddErrorIfFalse(CollectionContains(Collection0, Map0), TEXT("Unexpected OuterJoin iteration result"));
+
+						if (Map2 != nullptr)
+						{
+							AddErrorIfFalse(CollectionContains(Collection2, Map2), TEXT("Unexpected OuterJoin iteration result"));
+							AddErrorIfFalse(Map0 != nullptr && Map0->GetValueType() == Map2->GetValueType(), TEXT("Unexpected OuterJoin iteration result"));
+						}
+						else
+						{
+							NumNullEntries++;
+							AddErrorIfFalse(Map0 != nullptr, TEXT("Unexpected OuterJoin iteration result"));
+						}
+					},
+					Collection0.CreateMapIterator(),
+					Collection2.CreateConstMapIterator());
+				AddErrorIfFalse(NumPredicateCalls == 2, TEXT("Unexpected OuterJoin iteration result"));
+				AddErrorIfFalse(NumNullEntries == 1, TEXT("Unexpected OuterJoin iteration result"));
+			}
+
+			// Join with a transformer list (set maps)
+			{
+				// Use a float transformer
+				// We set up Collection0 such that it has two mappings with float values: Transform -> Float and Float -> Float
+				// InnerJoin should execute the predicate twice, once for each value type that matches the transformer's value type
+
+				// Dummy float sanitize transformer
+				struct FFloatAttributeTransformer_SanitizeTest
+				{
+					using TransformerType = Transformers::FSanitize;
+					using ValueType = FFloatAnimationAttribute;
+
+					static void TransformSetMap(FAttributeSetMap* SetMap) {}	// No-op
+				};
+
+				// Dummy vector sanitize transformer
+				struct FVectorAttributeTransformer_SanitizeTest
+				{
+					using TransformerType = Transformers::FSanitize;
+					using ValueType = FVectorAnimationAttribute;
+
+					static void TransformSetMap(FAttributeSetMap* SetMap) {}	// No-op
+				};
+
+#ifdef _MSC_VER
+	#pragma warning(push)
+	#pragma warning(disable : 4191) // unsafe conversion from 'function pointer type A' to 'function pointer type B'
+#endif
+
+				FAttributeValueTransformerList TransformerList(Transformers::FSanitize::TransformerName);
+				TransformerList.AddSetMapTransformer(FFloatAnimationAttribute::StaticStruct(), reinterpret_cast<FRawTransformerFunc>(&FFloatAttributeTransformer_SanitizeTest::TransformSetMap));
+				TransformerList.AddSetMapTransformer(FVectorAnimationAttribute::StaticStruct(), reinterpret_cast<FRawTransformerFunc>(&FVectorAttributeTransformer_SanitizeTest::TransformSetMap));
+
+#ifdef _MSC_VER
+	#pragma warning(pop)
+#endif
+
+				int32 NumPredicateCalls = 0;
+				int32 NumNullCollectionEntries = 0;
+				int32 NumNullTransformEntries = 0;
+				OuterJoinBy(
+					FAttributeValueTransformerListWithAttributeCollectionSetMapJoinOp(),
+					[this, &NumPredicateCalls, &NumNullCollectionEntries, &NumNullTransformEntries, &Collection0](FRawTransformerFunc Transformer, FAttributeSetMap* SetMap)
+					{
+						NumPredicateCalls++;
+
+						if (SetMap != nullptr)
+						{
+							AddErrorIfFalse(CollectionContains(Collection0, SetMap), TEXT("Unexpected OuterJoin iteration result"));
+
+							if (Transformer != nullptr)
+							{
+								AddErrorIfFalse(SetMap->GetValueType() == FFloatAnimationAttribute::StaticStruct(), TEXT("Unexpected OuterJoin iteration result"));
+							}
+						}
+						else
+						{
+							NumNullCollectionEntries++;
+						}
+
+						if (Transformer == nullptr)
+						{
+							NumNullTransformEntries++;
+						}
+					},
+					TransformerList.CreateSetMapTransformerIterator(),
+					Collection0.CreateSetMapIterator());
+
+				// [Transformer Value Type, Collection Set Map Value Type]
+				// ---
+				// [Float, Float->Float]
+				// [Float, Bone->Float]
+				// [Float, nullptr]
+				// [nullptr, Float->Int]
+				// [nullptr, Float->String]
+				// [nullptr, Bone->Bone]
+				// [Vector, nullptr]
+				
+				AddErrorIfFalse(NumPredicateCalls == 7, TEXT("Unexpected OuterJoin iteration result"));
+				AddErrorIfFalse(NumNullCollectionEntries == 2, TEXT("Unexpected OuterJoin iteration result"));
+				AddErrorIfFalse(NumNullTransformEntries == 3, TEXT("Unexpected OuterJoin iteration result"));
+			}
+
+			// Join with a transformer list (name/value maps)
+			{
+				// Use a float transformer
+				// We set up Collection0 such that it has two mappings with these value types: Float and Bone Transform
+				// InnerJoin should execute the predicate twice, once for each value type that matches the transformer's value type
+
+				// Dummy float sanitize transformer
+				struct FFloatAttributeTransformer_SanitizeTest
+				{
+					using TransformerType = Transformers::FSanitize;
+					using ValueType = FFloatAnimationAttribute;
+
+					static void TransformMap(FAttributeMap* Map) {}	// No-op
+				};
+
+				// Dummy vector sanitize transformer
+				struct FVectorAttributeTransformer_SanitizeTest
+				{
+					using TransformerType = Transformers::FSanitize;
+					using ValueType = FVectorAnimationAttribute;
+
+					static void TransformMap(FAttributeMap* Map) {}	// No-op
+				};
+
+#ifdef _MSC_VER
+	#pragma warning(push)
+	#pragma warning(disable : 4191) // unsafe conversion from 'function pointer type A' to 'function pointer type B'
+#endif
+
+				FAttributeValueTransformerList TransformerList(Transformers::FSanitize::TransformerName);
+				TransformerList.AddMapTransformer(FFloatAnimationAttribute::StaticStruct(), reinterpret_cast<FRawTransformerFunc>(&FFloatAttributeTransformer_SanitizeTest::TransformMap));
+				TransformerList.AddMapTransformer(FVectorAnimationAttribute::StaticStruct(), reinterpret_cast<FRawTransformerFunc>(&FVectorAttributeTransformer_SanitizeTest::TransformMap));
+
+#ifdef _MSC_VER
+	#pragma warning(pop)
+#endif
+
+				int32 NumPredicateCalls = 0;
+				int32 NumNullCollectionEntries = 0;
+				int32 NumNullTransformEntries = 0;
+				OuterJoinBy(
+					FAttributeValueTransformerListWithAttributeCollectionMapJoinOp(),
+					[this, &NumPredicateCalls, &NumNullCollectionEntries, &NumNullTransformEntries, &Collection0](FRawTransformerFunc Transformer, FAttributeMap* Map)
+					{
+						NumPredicateCalls++;
+
+						if (Map != nullptr)
+						{
+							AddErrorIfFalse(CollectionContains(Collection0, Map), TEXT("Unexpected OuterJoin iteration result"));
+
+							if (Transformer != nullptr)
+							{
+								AddErrorIfFalse(Map->GetValueType() == FFloatAnimationAttribute::StaticStruct(), TEXT("Unexpected OuterJoin iteration result"));
+							}
+						}
+						else
+						{
+							NumNullCollectionEntries++;
+						}
+
+						if (Transformer == nullptr)
+						{
+							NumNullTransformEntries++;
+						}
+					},
+					TransformerList.CreateMapTransformerIterator(),
+					Collection0.CreateMapIterator());
+
+				// [Transformer Value Type, Collection Map Value Type]
+				// ---
+				// [Float, Float]
+				// [nullptr, Bone]
+				// [Vector, nullptr]
+
+				AddErrorIfFalse(NumPredicateCalls == 3, TEXT("Unexpected OuterJoin iteration result"));
+				AddErrorIfFalse(NumNullCollectionEntries == 1, TEXT("Unexpected OuterJoin iteration result"));
+				AddErrorIfFalse(NumNullTransformEntries == 1, TEXT("Unexpected OuterJoin iteration result"));
+			}
+		}
+
+		return true;
+	}
+}
+
+#endif	// WITH_DEV_AUTOMATION_TESTS
