Rename UAF ValueRuntime types

diff --git a/Engine/Plugins/Experimental/UAF/UAF/Source/UAF/Public/UAF/ValueRuntime/Transformers/AdditiveSpace.h b/Engine/Plugins/Experimental/UAF/UAF/Source/UAF/Public/UAF/ValueRuntime/Transformers/AdditiveSpace.h
--- a/Engine/Plugins/Experimental/UAF/UAF/Source/UAF/Public/UAF/ValueRuntime/Transformers/AdditiveSpace.h
+++ b/Engine/Plugins/Experimental/UAF/UAF/Source/UAF/Public/UAF/ValueRuntime/Transformers/AdditiveSpace.h
@@ -12,25 +12,25 @@
 namespace UE::UAF::Transformers
 {
 	// Attribute value transformer to create additive values
-	struct FMakeAdditiveSpace final : public FAttributeValueTransformer
+	struct FMakeAdditiveSpace final : public FValueTransformer
 	{
 		static UE_API FName TransformerName;
 
 		// If the output buffer is nullptr, an allocator is provided and the implementation is responsible for allocating the output buffer and returning it
 		// If the base map is missing, transformers should use a default value
-		using FTransformSetMapFunc = void (*)(const FAttributeSetMap* Base, const FAttributeSetMap* Input, FAttributeSetMap*& Output, FReallocFun OutputAllocator);
-		using FTransformMapFunc = void (*)(const FAttributeMap* Base, const FAttributeMap* Input, FAttributeMap*& Output, FReallocFun OutputAllocator);
+		using FTransformBoundValueMapFunc = void (*)(const FBoundValueMap* Base, const FBoundValueMap* Input, FBoundValueMap*& Output, FReallocFun OutputAllocator);
+		using FTransformUnboundValueMapFunc = void (*)(const FUnboundValueMap* Base, const FUnboundValueMap* Input, FUnboundValueMap*& Output, FReallocFun OutputAllocator);
 
 		// Creates a new additive collection by subtracting the base from the specified Input
 		static UE_API void Apply(
-			const FAttributeValueTransformerMapPtr& TransformerMap,
+			const FValueTransformerMapPtr& TransformerMap,
 			const FAttributeCollection& Base,
 			const FAttributeCollection& Input,
 			FAttributeCollection& AdditiveOutput);
 	};
 
 	// Attribute value transformer to apply additive values onto a base
-	struct FApplyAdditiveSpace final : public FAttributeValueTransformer
+	struct FApplyAdditiveSpace final : public FValueTransformer
 	{
 		static UE_API FName TransformerName;
 
@@ -38,14 +38,14 @@ namespace UE::UAF::Transformers
 		// If the base map is missing, transformers should use a default value
 		// Per-value map is optional, when not present, the default weight should be used
 		// Unbound values do not support per-value weights
-		using FTransformSetMapFunc = void (*)(const FAttributeSetMap* Base, const FAttributeSetMap* Additive, const FAttributeSetMap* PerValueWeightMap, float DefaultWeight, FAttributeSetMap*& Output, FReallocFun OutputAllocator);
-		using FTransformMapFunc = void (*)(const FAttributeMap* Base, const FAttributeMap* Additive, float AdditiveWeight, FAttributeMap*& Output, FReallocFun OutputAllocator);
+		using FTransformBoundValueMapFunc = void (*)(const FBoundValueMap* Base, const FBoundValueMap* Additive, const FBoundValueMap* PerValueWeightMap, float DefaultWeight, FBoundValueMap*& Output, FReallocFun OutputAllocator);
+		using FTransformUnboundValueMapFunc = void (*)(const FUnboundValueMap* Base, const FUnboundValueMap* Additive, float AdditiveWeight, FUnboundValueMap*& Output, FReallocFun OutputAllocator);
 
 		// Applies an additive collection onto a base by adding it
 		// Output = Base + Lerp(Identity, Additive, Weight)
 		// This function supports outputting to either input
 		static UE_API void Apply(
-			const FAttributeValueTransformerMapPtr& TransformerMap,
+			const FValueTransformerMapPtr& TransformerMap,
 			const FAttributeCollection& Base,
 			const FAttributeCollection& Additive,
 			float AdditiveWeight,
@@ -56,7 +56,7 @@ namespace UE::UAF::Transformers
 		// Output = Base + Lerp(Identity, Additive, HasValueWeight ? ValueWeight : DefaultWeight)
 		// This function supports outputting to either input
 		static UE_API void Apply(
-			const FAttributeValueTransformerMapPtr& TransformerMap,
+			const FValueTransformerMapPtr& TransformerMap,
 			const FAttributeCollection& Base,
 			const FAttributeCollection& Additive,
 			const FAttributeCollection& PerValueWeights,
@@ -74,8 +74,8 @@ namespace UE::UAF::Transformers
 			using TransformerType = FMakeAdditiveSpace;
 			using ValueType = FBoneTransformAnimationAttribute;
 
-			static void TransformSetMap(const FAttributeSetMap* Base, const FAttributeSetMap* Input, FAttributeSetMap*& Output, FReallocFun OutputAllocator);
-			static void TransformMap(const FAttributeMap* Base, const FAttributeMap* Input, FAttributeMap*& Output, FReallocFun OutputAllocator);
+			static void TransformBoundValueMap(const FBoundValueMap* Base, const FBoundValueMap* Input, FBoundValueMap*& Output, FReallocFun OutputAllocator);
+			static void TransformUnboundValueMap(const FUnboundValueMap* Base, const FUnboundValueMap* Input, FUnboundValueMap*& Output, FReallocFun OutputAllocator);
 		};
 
 		struct FMakeAdditiveSpace_FloatAttribute
@@ -83,8 +83,8 @@ namespace UE::UAF::Transformers
 			using TransformerType = FMakeAdditiveSpace;
 			using ValueType = FFloatAnimationAttribute;
 
-			static void TransformSetMap(const FAttributeSetMap* Base, const FAttributeSetMap* Input, FAttributeSetMap*& Output, FReallocFun OutputAllocator);
-			static void TransformMap(const FAttributeMap* Base, const FAttributeMap* Input, FAttributeMap*& Output, FReallocFun OutputAllocator);
+			static void TransformBoundValueMap(const FBoundValueMap* Base, const FBoundValueMap* Input, FBoundValueMap*& Output, FReallocFun OutputAllocator);
+			static void TransformUnboundValueMap(const FUnboundValueMap* Base, const FUnboundValueMap* Input, FUnboundValueMap*& Output, FReallocFun OutputAllocator);
 		};
 
 		struct FApplyAdditiveSpace_BoneTransformAttribute
@@ -92,8 +92,8 @@ namespace UE::UAF::Transformers
 			using TransformerType = FApplyAdditiveSpace;
 			using ValueType = FBoneTransformAnimationAttribute;
 
-			static void TransformSetMap(const FAttributeSetMap* Base, const FAttributeSetMap* Additive, const FAttributeSetMap* PerValueWeightMap, float DefaultWeight, FAttributeSetMap*& Output, FReallocFun OutputAllocator);
-			static void TransformMap(const FAttributeMap* Base, const FAttributeMap* Additive, float AdditiveWeight, FAttributeMap*& Output, FReallocFun OutputAllocator);
+			static void TransformBoundValueMap(const FBoundValueMap* Base, const FBoundValueMap* Additive, const FBoundValueMap* PerValueWeightMap, float DefaultWeight, FBoundValueMap*& Output, FReallocFun OutputAllocator);
+			static void TransformUnboundValueMap(const FUnboundValueMap* Base, const FUnboundValueMap* Additive, float AdditiveWeight, FUnboundValueMap*& Output, FReallocFun OutputAllocator);
 		};
 
 		struct FApplyAdditiveSpace_FloatAttribute
@@ -101,8 +101,8 @@ namespace UE::UAF::Transformers
 			using TransformerType = FApplyAdditiveSpace;
 			using ValueType = FFloatAnimationAttribute;
 
-			static void TransformSetMap(const FAttributeSetMap* Base, const FAttributeSetMap* Additive, const FAttributeSetMap* PerValueWeightMap, float DefaultWeight, FAttributeSetMap*& Output, FReallocFun OutputAllocator);
-			static void TransformMap(const FAttributeMap* Base, const FAttributeMap* Additive, float AdditiveWeight, FAttributeMap*& Output, FReallocFun OutputAllocator);
+			static void TransformBoundValueMap(const FBoundValueMap* Base, const FBoundValueMap* Additive, const FBoundValueMap* PerValueWeightMap, float DefaultWeight, FBoundValueMap*& Output, FReallocFun OutputAllocator);
+			static void TransformUnboundValueMap(const FUnboundValueMap* Base, const FUnboundValueMap* Additive, float AdditiveWeight, FUnboundValueMap*& Output, FReallocFun OutputAllocator);
 		};
 	}
 }
diff --git a/Engine/Plugins/Experimental/UAF/UAF/Source/UAF/Public/UAF/ValueRuntime/Transformers/Interpolate.h b/Engine/Plugins/Experimental/UAF/UAF/Source/UAF/Public/UAF/ValueRuntime/Transformers/Interpolate.h
--- a/Engine/Plugins/Experimental/UAF/UAF/Source/UAF/Public/UAF/ValueRuntime/Transformers/Interpolate.h
+++ b/Engine/Plugins/Experimental/UAF/UAF/Source/UAF/Public/UAF/ValueRuntime/Transformers/Interpolate.h
@@ -13,7 +13,7 @@ namespace UE::UAF::Transformers
 {
 	// Attribute value transformer to interpolate values
 	// Make sure to also implement variations for Overwrite and Accumulate
-	struct FInterpolate final : public FAttributeValueTransformer
+	struct FInterpolate final : public FValueTransformer
 	{
 		static UE_API FName TransformerName;
 
@@ -21,20 +21,20 @@ namespace UE::UAF::Transformers
 		// Note that the input/output buffers can be the same buffer
 		// Per-value map is optional, when not present, the default weight should be used
 		// Unbound values do not support per-value weights
-		using FTransformSetMapFunc = void (*)(const FAttributeSetMap* SetMapA, const FAttributeSetMap* SetMapB, const FAttributeSetMap* PerValueWeightMapB, float DefaultWeightB, FAttributeSetMap*& OutputSetMap, FReallocFun OutputAllocator);
-		using FTransformMapFunc = void (*)(const FAttributeMap* MapA, const FAttributeMap* MapB, float WeightB, FAttributeMap*& OutputMap, FReallocFun OutputAllocator);
+		using FTransformBoundValueMapFunc = void (*)(const FBoundValueMap* MapA, const FBoundValueMap* MapB, const FBoundValueMap* PerValueWeightMapB, float DefaultWeightB, FBoundValueMap*& OutputMap, FReallocFun OutputAllocator);
+		using FTransformUnboundValueMapFunc = void (*)(const FUnboundValueMap* MapA, const FUnboundValueMap* MapB, float WeightB, FUnboundValueMap*& OutputMap, FReallocFun OutputAllocator);
 
 		// Interpolate attribute values
 		// Output = lerp(A, B, WeightB)
 		// Output = (A * (1.0 - WeightB)) + (B * WeightB)
 		// This function supports outputting to either input
-		static UE_API void Apply(const FAttributeValueTransformerMapPtr& TransformerMap, const FAttributeCollection& InputA, const FAttributeCollection& InputB, float WeightB, FAttributeCollection& Output);
+		static UE_API void Apply(const FValueTransformerMapPtr& TransformerMap, const FAttributeCollection& InputA, const FAttributeCollection& InputB, float WeightB, FAttributeCollection& Output);
 
 		// Interpolate attribute values
 		// If a weight isn't specified, the default weight will be used
 		// Output = lerp(A, B, HasValueWeightB ? ValueWeightB : DefaultWeightB)
 		// This function supports outputting to either input
-		static UE_API void Apply(const FAttributeValueTransformerMapPtr& TransformerMap, const FAttributeCollection& InputA, const FAttributeCollection& InputB, const FAttributeCollection& PerValueWeightsB, float DefaultWeightB, FAttributeCollection& Output);
+		static UE_API void Apply(const FValueTransformerMapPtr& TransformerMap, const FAttributeCollection& InputA, const FAttributeCollection& InputB, const FAttributeCollection& PerValueWeightsB, float DefaultWeightB, FAttributeCollection& Output);
 	};
 
 	//////////////////////////////////////////////////////////////////////////
@@ -47,8 +47,8 @@ namespace UE::UAF::Transformers
 			using TransformerType = FInterpolate;
 			using ValueType = FBoneTransformAnimationAttribute;
 
-			static void TransformSetMap(const FAttributeSetMap* SetMapA, const FAttributeSetMap* SetMapB, const FAttributeSetMap* PerValueWeightMapB, float DefaultWeightB, FAttributeSetMap*& OutputSetMap, FReallocFun OutputAllocator);
-			static void TransformMap(const FAttributeMap* MapA, const FAttributeMap* MapB, float WeightB, FAttributeMap*& OutputMap, FReallocFun OutputAllocator);
+			static void TransformBoundValueMap(const FBoundValueMap* MapA, const FBoundValueMap* MapB, const FBoundValueMap* PerValueWeightMapB, float DefaultWeightB, FBoundValueMap*& OutputMap, FReallocFun OutputAllocator);
+			static void TransformUnboundValueMap(const FUnboundValueMap* MapA, const FUnboundValueMap* MapB, float WeightB, FUnboundValueMap*& OutputMap, FReallocFun OutputAllocator);
 		};
 
 		struct FInterpolate_FloatAttribute
@@ -56,8 +56,8 @@ namespace UE::UAF::Transformers
 			using TransformerType = FInterpolate;
 			using ValueType = FFloatAnimationAttribute;
 
-			static void TransformSetMap(const FAttributeSetMap* SetMapA, const FAttributeSetMap* SetMapB, const FAttributeSetMap* PerValueWeightMapB, float DefaultWeightB, FAttributeSetMap*& OutputSetMap, FReallocFun OutputAllocator);
-			static void TransformMap(const FAttributeMap* MapA, const FAttributeMap* MapB, float WeightB, FAttributeMap*& OutputMap, FReallocFun OutputAllocator);
+			static void TransformBoundValueMap(const FBoundValueMap* MapA, const FBoundValueMap* MapB, const FBoundValueMap* PerValueWeightMapB, float DefaultWeightB, FBoundValueMap*& OutputMap, FReallocFun OutputAllocator);
+			static void TransformUnboundValueMap(const FUnboundValueMap* MapA, const FUnboundValueMap* MapB, float WeightB, FUnboundValueMap*& OutputMap, FReallocFun OutputAllocator);
 		};
 	}
 }
diff --git a/Engine/Plugins/Experimental/UAF/UAF/Source/UAF/Public/UAF/ValueRuntime/Transformers/Layer.h b/Engine/Plugins/Experimental/UAF/UAF/Source/UAF/Public/UAF/ValueRuntime/Transformers/Layer.h
--- a/Engine/Plugins/Experimental/UAF/UAF/Source/UAF/Public/UAF/ValueRuntime/Transformers/Layer.h
+++ b/Engine/Plugins/Experimental/UAF/UAF/Source/UAF/Public/UAF/ValueRuntime/Transformers/Layer.h
@@ -12,7 +12,7 @@
 namespace UE::UAF::Transformers
 {
 	// Attribute value transformer to layer values
-	struct FLayer final : public FAttributeValueTransformer
+	struct FLayer final : public FValueTransformer
 	{
 		static UE_API FName TransformerName;
 
@@ -20,8 +20,8 @@ namespace UE::UAF::Transformers
 		// Note that the input/output buffers can be the same buffer
 		// Per-value map is optional, when not present, the default weight should be used
 		// Unbound values do not support per-value weights
-		using FTransformSetMapFunc = void (*)(const FAttributeSetMap* Base, const FAttributeSetMap* Layer, const FAttributeSetMap* PerValueWeights, float DefaultLayerWeight, FAttributeSetMap*& OutputSetMap, FReallocFun OutputAllocator);
-		using FTransformMapFunc = void (*)(const FAttributeMap* Base, const FAttributeMap* Layer, float LayerWeight, FAttributeMap*& OutputMap, FReallocFun OutputAllocator);
+		using FTransformBoundValueMapFunc = void (*)(const FBoundValueMap* Base, const FBoundValueMap* Layer, const FBoundValueMap* PerValueWeights, float DefaultLayerWeight, FBoundValueMap*& OutputMap, FReallocFun OutputAllocator);
+		using FTransformUnboundValueMapFunc = void (*)(const FUnboundValueMap* Base, const FUnboundValueMap* Layer, float LayerWeight, FUnboundValueMap*& OutputMap, FReallocFun OutputAllocator);
 
 		// Layer values onto a base
 		// Output = layer(Base, Layer, LayerWeight)
@@ -29,7 +29,7 @@ namespace UE::UAF::Transformers
 		// Layer values do not have to be present in the Base and extra values are dropped
 		// This function supports outputting to either input
 		static UE_API void Apply(
-			const FAttributeValueTransformerMapPtr& TransformerMap,
+			const FValueTransformerMapPtr& TransformerMap,
 			const FAttributeCollection& Base,
 			const FAttributeCollection& Layer,
 			float LayerWeight,
@@ -40,7 +40,7 @@ namespace UE::UAF::Transformers
 		// Layer values do not have to be present in the Base and extra values are dropped
 		// This function supports outputting to either input
 		static UE_API void Apply(
-			const FAttributeValueTransformerMapPtr& TransformerMap,
+			const FValueTransformerMapPtr& TransformerMap,
 			const FAttributeCollection& Base,
 			const FAttributeCollection& Layer,
 			const FAttributeCollection& PerValueLayerWeights,
@@ -58,8 +58,8 @@ namespace UE::UAF::Transformers
 			using TransformerType = FLayer;
 			using ValueType = FBoneTransformAnimationAttribute;
 
-			static void TransformSetMap(const FAttributeSetMap* Base, const FAttributeSetMap* Layer, const FAttributeSetMap* PerValueWeights, float DefaultLayerWeight, FAttributeSetMap*& OutputSetMap, FReallocFun OutputAllocator);
-			static void TransformMap(const FAttributeMap* Base, const FAttributeMap* Layer, float LayerWeight, FAttributeMap*& OutputMap, FReallocFun OutputAllocator);
+			static void TransformBoundValueMap(const FBoundValueMap* Base, const FBoundValueMap* Layer, const FBoundValueMap* PerValueWeights, float DefaultLayerWeight, FBoundValueMap*& OutputMap, FReallocFun OutputAllocator);
+			static void TransformUnboundValueMap(const FUnboundValueMap* Base, const FUnboundValueMap* Layer, float LayerWeight, FUnboundValueMap*& OutputMap, FReallocFun OutputAllocator);
 		};
 
 		struct FLayer_FloatAttribute
@@ -67,8 +67,8 @@ namespace UE::UAF::Transformers
 			using TransformerType = FLayer;
 			using ValueType = FFloatAnimationAttribute;
 
-			static void TransformSetMap(const FAttributeSetMap* Base, const FAttributeSetMap* Layer, const FAttributeSetMap* PerValueWeights, float DefaultLayerWeight, FAttributeSetMap*& OutputSetMap, FReallocFun OutputAllocator);
-			static void TransformMap(const FAttributeMap* Base, const FAttributeMap* Layer, float LayerWeight, FAttributeMap*& OutputMap, FReallocFun OutputAllocator);
+			static void TransformBoundValueMap(const FBoundValueMap* Base, const FBoundValueMap* Layer, const FBoundValueMap* PerValueWeights, float DefaultLayerWeight, FBoundValueMap*& OutputMap, FReallocFun OutputAllocator);
+			static void TransformUnboundValueMap(const FUnboundValueMap* Base, const FUnboundValueMap* Layer, float LayerWeight, FUnboundValueMap*& OutputMap, FReallocFun OutputAllocator);
 		};
 	}
 }
diff --git a/Engine/Plugins/Experimental/UAF/UAF/Source/UAF/Public/UAF/ValueRuntime/Transformers/Overwrite.h b/Engine/Plugins/Experimental/UAF/UAF/Source/UAF/Public/UAF/ValueRuntime/Transformers/Overwrite.h
--- a/Engine/Plugins/Experimental/UAF/UAF/Source/UAF/Public/UAF/ValueRuntime/Transformers/Overwrite.h
+++ b/Engine/Plugins/Experimental/UAF/UAF/Source/UAF/Public/UAF/ValueRuntime/Transformers/Overwrite.h
@@ -14,7 +14,7 @@ namespace UE::UAF::Transformers
 	// Attribute value transformer to overwrite values
 	// Overwrite is the initial step within an N-step interpolation: Overwrite, Accumulate0, Accumulate1, ...
 	// Make sure to also implement variations for Accumulate and Interpolate
-	struct FOverwrite final : public FAttributeValueTransformer
+	struct FOverwrite final : public FValueTransformer
 	{
 		static UE_API FName TransformerName;
 
@@ -22,19 +22,19 @@ namespace UE::UAF::Transformers
 		// Note that the input/output buffers can be the same buffer
 		// Per-value map is optional, when not present, the default weight should be used
 		// Unbound values do not support per-value weights
-		using FTransformSetMapFunc = void (*)(const FAttributeSetMap* InputSetMap, const FAttributeSetMap* PerValueWeightMap, float DefaultWeight, FAttributeSetMap*& OutputSetMap, FReallocFun OutputAllocator);
-		using FTransformMapFunc = void (*)(const FAttributeMap* InputMap, float Weight, FAttributeMap*& OutputMap, FReallocFun OutputAllocator);
+		using FTransformBoundValueMapFunc = void (*)(const FBoundValueMap* InputMap, const FBoundValueMap* PerValueWeightMap, float DefaultWeight, FBoundValueMap*& OutputMap, FReallocFun OutputAllocator);
+		using FTransformUnboundValueMapFunc = void (*)(const FUnboundValueMap* InputMap, float Weight, FUnboundValueMap*& OutputMap, FReallocFun OutputAllocator);
 
 		// Overwrite attribute values
 		// Output = (A * Weight)
 		// This function supports outputting to its input
-		static UE_API void Apply(const FAttributeValueTransformerMapPtr& TransformerMap, const FAttributeCollection& Input, float Weight, FAttributeCollection& Output);
+		static UE_API void Apply(const FValueTransformerMapPtr& TransformerMap, const FAttributeCollection& Input, float Weight, FAttributeCollection& Output);
 
 		// Overwrite attribute values using per value weights (must be of type FFloatAnimationAttribute)
 		// If a weight isn't specified, the default weight will be used
 		// Output = A * (HasValueWeight ? ValueWeight : DefaultWeight)
 		// This function supports outputting to its input
-		static UE_API void Apply(const FAttributeValueTransformerMapPtr& TransformerMap, const FAttributeCollection& Input, const FAttributeCollection& PerValueWeights, float DefaultWeight, FAttributeCollection& Output);
+		static UE_API void Apply(const FValueTransformerMapPtr& TransformerMap, const FAttributeCollection& Input, const FAttributeCollection& PerValueWeights, float DefaultWeight, FAttributeCollection& Output);
 	};
 
 	//////////////////////////////////////////////////////////////////////////
@@ -47,8 +47,8 @@ namespace UE::UAF::Transformers
 			using TransformerType = FOverwrite;
 			using ValueType = FBoneTransformAnimationAttribute;
 
-			static void TransformSetMap(const FAttributeSetMap* InputSetMap, const FAttributeSetMap* PerValueWeightMap, float DefaultWeight, FAttributeSetMap*& OutputSetMap, FReallocFun OutputAllocator);
-			static void TransformMap(const FAttributeMap* InputMap, float Weight, FAttributeMap*& OutputMap, FReallocFun OutputAllocator);
+			static void TransformBoundValueMap(const FBoundValueMap* InputMap, const FBoundValueMap* PerValueWeightMap, float DefaultWeight, FBoundValueMap*& OutputMap, FReallocFun OutputAllocator);
+			static void TransformUnboundValueMap(const FUnboundValueMap* InputMap, float Weight, FUnboundValueMap*& OutputMap, FReallocFun OutputAllocator);
 		};
 
 		struct FOverwrite_FloatAttribute
@@ -56,8 +56,8 @@ namespace UE::UAF::Transformers
 			using TransformerType = FOverwrite;
 			using ValueType = FFloatAnimationAttribute;
 
-			static void TransformSetMap(const FAttributeSetMap* InputSetMap, const FAttributeSetMap* PerValueWeightMap, float DefaultWeight, FAttributeSetMap*& OutputSetMap, FReallocFun OutputAllocator);
-			static void TransformMap(const FAttributeMap* InputMap, float Weight, FAttributeMap*& OutputMap, FReallocFun OutputAllocator);
+			static void TransformBoundValueMap(const FBoundValueMap* InputMap, const FBoundValueMap* PerValueWeightMap, float DefaultWeight, FBoundValueMap*& OutputMap, FReallocFun OutputAllocator);
+			static void TransformUnboundValueMap(const FUnboundValueMap* InputMap, float Weight, FUnboundValueMap*& OutputMap, FReallocFun OutputAllocator);
 		};
 	}
 }
diff --git a/Engine/Plugins/Experimental/UAF/UAF/Source/UAF/Public/UAF/ValueRuntime/Transformers/Sanitize.h b/Engine/Plugins/Experimental/UAF/UAF/Source/UAF/Public/UAF/ValueRuntime/Transformers/Sanitize.h
--- a/Engine/Plugins/Experimental/UAF/UAF/Source/UAF/Public/UAF/ValueRuntime/Transformers/Sanitize.h
+++ b/Engine/Plugins/Experimental/UAF/UAF/Source/UAF/Public/UAF/ValueRuntime/Transformers/Sanitize.h
@@ -16,16 +16,16 @@ namespace UE::UAF::Transformers
 {
 	// Attribute value transformer to sanitize values
 	// This can be used to normalize rotations and other like values as well as check for NaN/Inf values
-	struct FSanitize final : public FAttributeValueTransformer
+	struct FSanitize final : public FValueTransformer
 	{
 		static UE_API FName TransformerName;
 
-		using FTransformSetMapFunc = void (*)(FAttributeSetMap* SetMap);
-		using FTransformMapFunc = void (*)(FAttributeMap* Map);
+		using FTransformBoundValueMapFunc = void (*)(FBoundValueMap* Map);
+		using FTransformUnboundValueMapFunc = void (*)(FUnboundValueMap* Map);
 
 		// Sanitize attribute values
 		// This can be used to normalize rotations (and other like values) as well as check for NaN/Inf values
-		static UE_API void Apply(const FAttributeValueTransformerMapPtr& TransformerMap, FAttributeCollection& Collection);
+		static UE_API void Apply(const FValueTransformerMapPtr& TransformerMap, FAttributeCollection& Collection);
 	};
 
 	//////////////////////////////////////////////////////////////////////////
@@ -39,8 +39,8 @@ namespace UE::UAF::Transformers
 			using TransformerType = FSanitize;
 			using ValueType = FBoneTransformAnimationAttribute;
 
-			static void TransformSetMap(FAttributeSetMap* SetMap);
-			static void TransformMap(FAttributeMap* Map);
+			static void TransformBoundValueMap(FBoundValueMap* Map);
+			static void TransformUnboundValueMap(FUnboundValueMap* Map);
 		};
 
 #if UE_UAF_VALIDATE_MAPPED_VALUES
@@ -50,8 +50,8 @@ namespace UE::UAF::Transformers
 			using TransformerType = FSanitize;
 			using ValueType = FFloatAnimationAttribute;
 
-			static void TransformSetMap(FAttributeSetMap* SetMap);
-			static void TransformMap(FAttributeMap* Map);
+			static void TransformBoundValueMap(FBoundValueMap* Map);
+			static void TransformUnboundValueMap(FUnboundValueMap* Map);
 		};
 #endif
 	}
diff --git a/Engine/Plugins/Experimental/UAF/UAF/Source/UAF/Public/UAF/ValueRuntime/Transformers/Transformer.h b/Engine/Plugins/Experimental/UAF/UAF/Source/UAF/Public/UAF/ValueRuntime/Transformers/Transformer.h
--- a/Engine/Plugins/Experimental/UAF/UAF/Source/UAF/Public/UAF/ValueRuntime/Transformers/Transformer.h
+++ b/Engine/Plugins/Experimental/UAF/UAF/Source/UAF/Public/UAF/ValueRuntime/Transformers/Transformer.h
@@ -10,34 +10,34 @@
 namespace UE::UAF
 {
 	class FValueRuntimeRegistry;
-	class FAttributeValueTransformerList;
+	class FValueTransformerList;
 
 	// A raw function pointer for all transformers (static functions)
 	// Only used for storage purposes, NEVER CALL THIS!
 	using FRawTransformerFunc = void(*)();
 
-	// Attribute value transformer base type
+	// Value transformer base type
 	// All transformer types must derive from this
-	struct FAttributeValueTransformer
+	struct FValueTransformer
 	{
 		// The name of the transformer (exposed in derived types)
 		//static FName TransformerName;
 
-		// An optional function signature of the transformer for set maps (exposed in derived types)
-		//using FTransformSetMapFunc = void (*)(...);
+		// An optional function signature of the transformer for bound value maps (exposed in derived types)
+		//using FTransformBoundValueMapFunc = void (*)(...);
 
-		// An optional function signature of the transformer for name/value maps (exposed in derived types)
-		//using FTransformMapFunc = void (*)(...);
+		// An optional function signature of the transformer for unbound value maps (exposed in derived types)
+		//using FTransformUnboundValueMapFunc = void (*)(...);
 	};
 
-	// An iterator over a transformer list within the attribute runtime registry
-	class FAttributeValueTransformerListIterator
+	// An iterator over a transformer list within the value runtime registry
+	class FValueTransformerListIterator
 	{
 	public:
-		[[nodiscard]] FAttributeValueTransformerListIterator(const FAttributeValueTransformerList& TransformerList, int32 StartOffset, int32 Count);
+		[[nodiscard]] FValueTransformerListIterator(const FValueTransformerList& TransformerList, int32 StartOffset, int32 Count);
 
 		// Increments and moves the iterator to the next transformer entry
-		FAttributeValueTransformerListIterator& operator++();
+		FValueTransformerListIterator& operator++();
 
 		// Returns whether or not the iterator still contains values
 		[[nodiscard]] explicit operator bool() const;
@@ -56,64 +56,64 @@ namespace UE::UAF
 
 	// A list of transformer specializations that operators over various value types.
 	// Transformers have a base type (e.g. interpolate, sanitize) and can be specialized to operate on arbitrary
-	// value types (e.g. FFloatAnimationAttribute). They can also transform two types of containers: set maps
-	// and name/value maps.
-	class FAttributeValueTransformerList
+	// value types (e.g. FFloatAnimationAttribute). They can also transform two types of containers: bound value maps
+	// and unbound value maps.
+	class FValueTransformerList
 	{
 	public:
 		// Creates an empty transformer list
-		UE_API FAttributeValueTransformerList();
+		UE_API FValueTransformerList();
 
 		// Creates a transformer list for the specified transformer type
-		[[nodiscard]] UE_API explicit FAttributeValueTransformerList(FName TransformerName);
+		[[nodiscard]] UE_API explicit FValueTransformerList(FName TransformerName);
 
 		// Returns the transformer name
 		[[nodiscard]] UE_API FName GetTransformerName() const;
 
-		// Returns the number of set map transformers in this list
-		[[nodiscard]] int32 NumSetMapTransformers() const { return SetMapTransformerCount; }
+		// Returns the number of bound value map transformers in this list
+		[[nodiscard]] int32 NumBoundValueMapTransformers() const { return BoundValueMapTransformerCount; }
 
-		// Returns the number of name/value map transformers in this list
-		[[nodiscard]] int32 NumMapTransformers() const { return ValueTypes.Num() - SetMapTransformerCount; }
+		// Returns the number of unbound value map transformers in this list
+		[[nodiscard]] int32 NumUnboundValueMapTransformers() const { return ValueTypes.Num() - BoundValueMapTransformerCount; }
 
-		// Adds a new set map transformer entry
-		UE_API bool AddSetMapTransformer(UScriptStruct* ValueType, FRawTransformerFunc TransformerFunc);
+		// Adds a new bound value map transformer entry
+		UE_API bool AddBoundValueMapTransformer(UScriptStruct* ValueType, FRawTransformerFunc TransformerFunc);
 
-		// Adds a new name/value map transformer entry
-		UE_API bool AddMapTransformer(UScriptStruct* ValueType, FRawTransformerFunc TransformerFunc);
+		// Adds a new unbound value map transformer entry
+		UE_API bool AddUnboundValueMapTransformer(UScriptStruct* ValueType, FRawTransformerFunc TransformerFunc);
 
-		// Returns an iterator over all set map transformer specializations within
-		[[nodiscard]] FAttributeValueTransformerListIterator CreateSetMapTransformerIterator() const { return FAttributeValueTransformerListIterator(*this, 0, NumSetMapTransformers()); }
+		// Returns an iterator over all bound value map transformer specializations within
+		[[nodiscard]] FValueTransformerListIterator CreateBoundValueMapTransformerIterator() const { return FValueTransformerListIterator(*this, 0, NumBoundValueMapTransformers()); }
 
-		// Returns an iterator over all name/value map transformer specializations within
-		[[nodiscard]] FAttributeValueTransformerListIterator CreateMapTransformerIterator() const { return FAttributeValueTransformerListIterator(*this, NumSetMapTransformers(), NumMapTransformers()); }
+		// Returns an iterator over all unbound value map transformer specializations within
+		[[nodiscard]] FValueTransformerListIterator CreateUnboundValueMapTransformerIterator() const { return FValueTransformerListIterator(*this, NumBoundValueMapTransformers(), NumUnboundValueMapTransformers()); }
 
 	private:
 		// The type of the transformer (e.g. interpolate, sanitize)
 		FName TransformerName;
 
-		// The number of set map transformers in this list
-		int32 SetMapTransformerCount = 0;
+		// The number of bound value map transformers in this list
+		int32 BoundValueMapTransformerCount = 0;
 
-		// List of attribute value types for which we have transformers registered (e.g. FTransformAnimationAttribute, FFloatAnimationAttribute) (sorted by container type, then by value type)
-		// This list is split in two, set map transformers come first, then map transformers
+		// List of value types for which we have transformers registered (e.g. FTransformAnimationAttribute, FFloatAnimationAttribute) (sorted by container type, then by value type)
+		// This list is split in two, bound value map transformers come first, then unbound value map transformers
 		// Each sub-list is sorted by value type
 		TArray<UScriptStruct*> ValueTypes;
 
 		// List of transformer functions for each value type (sorted by container type, then by value type)
-		// This list is split in two, set map transformers come first, then map transformers
+		// This list is split in two, bound value map transformers come first, then unbound value map transformers
 		// Each sub-list is sorted by value type
 		TArray<FRawTransformerFunc> Transformers;
 
 		friend FValueRuntimeRegistry;
-		friend FAttributeValueTransformerListIterator;
+		friend FValueTransformerListIterator;
 	};
 
-	using FAttributeValueTransformerMapPtr = TSharedPtr<TMap<FName, FAttributeValueTransformerList>>;
+	using FValueTransformerMapPtr = TSharedPtr<TMap<FName, FValueTransformerList>>;
 
-	// An operator structure to join attribute value transformer lists and attribute collections using InnerJoin/OuterJoin/etc
-	// Ensures that the predicate is called for each set map with a value type that has a corresponding transformer
-	struct FAttributeValueTransformerListWithAttributeCollectionSetMapJoinOp
+	// An operator structure to join value transformer lists and attribute collections using InnerJoin/OuterJoin/etc
+	// Ensures that the predicate is called for each bound value map with a value type that has a corresponding transformer
+	struct FValueTransformerListWithBoundMapCollectionJoinOp
 	{
 		// Key sorting uses signed comparison, use the largest signed value regardless of pointer size
 		static constexpr intptr_t WILDCARD_ATTRIBUTE_TYPE_VALUE = static_cast<intptr_t>(static_cast<uintptr_t>(~0ULL) >> 1);
@@ -136,25 +136,26 @@ namespace UE::UAF
 		FAttributeMappingKey GetLargestKey() const { return FAttributeMappingKey::LARGEST_VALUE; }
 
 		// FAttributeValueTransformerListIterator entries
-		FAttributeMappingKey GetKey(const FAttributeValueTransformerListIterator& It) const
+		FAttributeMappingKey GetKey(const FValueTransformerListIterator& It) const
 		{
 			// We use a wildcard entry for the attribute type within our key to ensure it is always larger than all other valid values
 			return FAttributeMappingKey(reinterpret_cast<UScriptStruct*>(WILDCARD_ATTRIBUTE_TYPE_VALUE), It.GetValueType());
 		}
 
-		FRawTransformerFunc GetValue(const FAttributeValueTransformerListIterator& It) const { return It.GetTransformerFunc(); }
+		FRawTransformerFunc GetValue(const FValueTransformerListIterator& It) const { return It.GetTransformerFunc(); }
 
-		// Normal FAttributeCollectionSetMapJoinOp
-		FAttributeMappingKey GetKey(const FAttributeCollection::FSetMapIterator& It) const { return It.GetKey(); }
-		FAttributeMappingKey GetKey(const FAttributeCollection::FConstSetMapIterator& It) const { return It.GetKey(); }
+		// Normal FBoundMapCollectionJoinOp
+		template<class IteratorType>
+		FAttributeMappingKey GetKey(const IteratorType& It) const { return It.GetKey(); }
 
-		FAttributeSetMap* GetValue(const FAttributeCollection::FSetMapIterator& It) const { return It.GetMap(); }
-		const FAttributeSetMap* GetValue(const FAttributeCollection::FConstSetMapIterator& It) const { return It.GetMap(); }
+		template<class IteratorType>
+			requires requires (const IteratorType& It) { It.GetMap(); }
+		typename IteratorType::MapType GetValue(const IteratorType& It) const { return It.GetMap(); }
 	};
 
-	// An operator structure to join attribute value transformer lists and attribute collections using InnerJoin/OuterJoin/etc
-	// Ensures that the predicate is called for each set map with a value type that has a corresponding transformer
-	struct FAttributeValueTransformerListWithAttributeCollectionMapJoinOp
+	// An operator structure to join value transformer lists and attribute collections using InnerJoin/OuterJoin/etc
+	// Ensures that the predicate is called for each unbound value map with a value type that has a corresponding transformer
+	struct FValueTransformerListWithUnboundMapCollectionJoinOp
 	{
 		// Key sorting uses signed comparison, use the largest signed value regardless of pointer size
 		static constexpr intptr_t KEY_LARGEST_VALUE = static_cast<intptr_t>(static_cast<uintptr_t>(~0ULL) >> 1);
@@ -162,28 +163,29 @@ namespace UE::UAF
 		UScriptStruct* GetLargestKey() const { return reinterpret_cast<UScriptStruct*>(KEY_LARGEST_VALUE); }
 
 		// FAttributeValueTransformerListIterator entries
-		UScriptStruct* GetKey(const FAttributeValueTransformerListIterator& It) const { return It.GetValueType(); }
-		FRawTransformerFunc GetValue(const FAttributeValueTransformerListIterator& It) const { return It.GetTransformerFunc(); }
+		UScriptStruct* GetKey(const FValueTransformerListIterator& It) const { return It.GetValueType(); }
+		FRawTransformerFunc GetValue(const FValueTransformerListIterator& It) const { return It.GetTransformerFunc(); }
 
-		// Normal FAttributeCollectionMapJoinOp
-		UScriptStruct* GetKey(const FAttributeCollection::FMapIterator& It) const { return It.GetValueType(); }
-		UScriptStruct* GetKey(const FAttributeCollection::FConstMapIterator& It) const { return It.GetValueType(); }
+		// Normal FUnboundMapCollectionJoinOp
+		template<class IteratorType>
+		UScriptStruct* GetKey(const IteratorType& It) const { return It.GetValueType(); }
 
-		FAttributeMap* GetValue(const FAttributeCollection::FMapIterator& It) const { return It.GetMap(); }
-		const FAttributeMap* GetValue(const FAttributeCollection::FConstMapIterator& It) const { return It.GetMap(); }
+		template<class IteratorType>
+			requires requires (const IteratorType& It) { It.GetMap(); }
+		typename IteratorType::MapType GetValue(const IteratorType& It) const { return It.GetMap(); }
 	};
 
 	//////////////////////////////////////////////////////////////////////////
 	// Implementation
 
-	inline FAttributeValueTransformerListIterator::FAttributeValueTransformerListIterator(const FAttributeValueTransformerList& TransformerList, int32 StartOffset, int32 Count)
+	inline FValueTransformerListIterator::FValueTransformerListIterator(const FValueTransformerList& TransformerList, int32 StartOffset, int32 Count)
 		: ValueTypePtr(TransformerList.ValueTypes.GetData() + StartOffset)
 		, TransformerFuncPtr(TransformerList.Transformers.GetData() + StartOffset)
 		, ValueTypeEndPtr(ValueTypePtr + Count)
 	{
 	}
 
-	inline FAttributeValueTransformerListIterator& FAttributeValueTransformerListIterator::operator++()
+	inline FValueTransformerListIterator& FValueTransformerListIterator::operator++()
 	{
 		++ValueTypePtr;
 		++TransformerFuncPtr;
@@ -191,17 +193,17 @@ namespace UE::UAF
 		return *this;
 	}
 
-	inline FAttributeValueTransformerListIterator::operator bool() const
+	inline FValueTransformerListIterator::operator bool() const
 	{
 		return ValueTypePtr < ValueTypeEndPtr;
 	}
 
-	inline UScriptStruct* FAttributeValueTransformerListIterator::GetValueType() const
+	inline UScriptStruct* FValueTransformerListIterator::GetValueType() const
 	{
 		return *ValueTypePtr;
 	}
 
-	inline FRawTransformerFunc FAttributeValueTransformerListIterator::GetTransformerFunc() const
+	inline FRawTransformerFunc FValueTransformerListIterator::GetTransformerFunc() const
 	{
 		return *TransformerFuncPtr;
 	}
diff --git a/Engine/Plugins/Experimental/UAF/UAF/Source/UAF/Public/UAF/ValueRuntime/ValueRuntimeRegistry.h b/Engine/Plugins/Experimental/UAF/UAF/Source/UAF/Public/UAF/ValueRuntime/ValueRuntimeRegistry.h
--- a/Engine/Plugins/Experimental/UAF/UAF/Source/UAF/Public/UAF/ValueRuntime/ValueRuntimeRegistry.h
+++ b/Engine/Plugins/Experimental/UAF/UAF/Source/UAF/Public/UAF/ValueRuntime/ValueRuntimeRegistry.h
@@ -3,9 +3,9 @@
 #pragma once
 
 #include "CoreMinimal.h"
-#include "UAF/ValueRuntime/AttributeMap.h"
-#include "UAF/ValueRuntime/AttributeSetMap.h"
+#include "UAF/ValueRuntime/BoundValueMap.h"
 #include "UAF/ValueRuntime/Transformers/Transformer.h"
+#include "UAF/ValueRuntime/UnboundValueMap.h"
 
 #define UE_API UAF_API
 
@@ -17,55 +17,55 @@ namespace UE::UAF
 	public:
 		UE_NONCOPYABLE(FValueRuntimeRegistry);	// We disallow copy/move semantics
 
-		using SetMapConstructorFun = TFunction<FAttributeSetMap*(const FAttributeSetMap::FConstructArgs& Args)>;
-		using MapConstructorFun = TFunction<FAttributeMap*(FReallocFun ReallocFun)>;
+		using BoundValueMapConstructorFun = TFunction<FBoundValueMap*(const FBoundValueMap::FConstructArgs& Args)>;
+		using UnboundValueMapConstructorFun = TFunction<FUnboundValueMap*(FReallocFun ReallocFun)>;
 
-		// Returns the set map registry instance
+		// Returns the value runtime registry instance
 		[[nodiscard]] static UE_API FValueRuntimeRegistry& Get();
 
-		// Constructs a set map from its UStruct value type
-		[[nodiscard]] UE_API FAttributeSetMap* ConstructSetMap(const FAttributeSetMap::FConstructArgs& Args) const;
+		// Constructs a bound value map from its UStruct value type
+		[[nodiscard]] UE_API FBoundValueMap* ConstructBoundValueMap(const FBoundValueMap::FConstructArgs& Args) const;
 
 		// Constructs a name/value map from its UStruct value type
-		[[nodiscard]] UE_API FAttributeMap* ConstructMap(UScriptStruct* ValueType, FReallocFun ReallocFun) const;
+		[[nodiscard]] UE_API FUnboundValueMap* ConstructUnboundValueMap(UScriptStruct* ValueType, FReallocFun ReallocFun) const;
 
-		// Returns the map of known attribute value transformers
-		[[nodiscard]] UE_API FAttributeValueTransformerMapPtr GetTransformerMap() const;
+		// Returns the map of known value transformers
+		[[nodiscard]] UE_API FValueTransformerMapPtr GetTransformerMap() const;
 
-		// Registers a set map of the specified value type
+		// Registers a bound value map of the specified value type
 		template<class ValueType>
-		void RegisterSetMapInitializer();
+		void RegisterBoundValueMapInitializer();
 
 		// Registers a name/value map of the specified value type
 		template<class ValueType>
-		void RegisterMapInitializer();
+		void RegisterUnboundValueMapInitializer();
 
 		template<class TransformerSpecializationType>
-		bool RegisterAttributeValueTransformer();
+		bool RegisterValueTransformer();
 
 	private:
 		UE_API FValueRuntimeRegistry();
-		UE_API void RegisterSetMapInitializer(UScriptStruct* ValueType, SetMapConstructorFun Initializer);
-		UE_API void RegisterMapInitializer(UScriptStruct* ValueType, MapConstructorFun Initializer);
-		UE_API bool RegisterAttributeValueSetMapTransformer(FName TransformerName, UScriptStruct* ValueType, FRawTransformerFunc TransformerFunc);
-		UE_API bool RegisterAttributeValueMapTransformer(FName TransformerName, UScriptStruct* ValueType, FRawTransformerFunc TransformerFunc);
+		UE_API void RegisterBoundValueMapInitializer(UScriptStruct* ValueType, BoundValueMapConstructorFun Initializer);
+		UE_API void RegisterUnboundValueMapInitializer(UScriptStruct* ValueType, UnboundValueMapConstructorFun Initializer);
+		UE_API bool RegisterBouneValueMapTransformer(FName TransformerName, UScriptStruct* ValueType, FRawTransformerFunc TransformerFunc);
+		UE_API bool RegisterUnboundValueMapTransformer(FName TransformerName, UScriptStruct* ValueType, FRawTransformerFunc TransformerFunc);
 
 		// Initialize/destroy the global registry
 		static UE_API void Init();
 		static UE_API void Destroy();
 
-		// A map of ValueType -> Set Map constructors
+		// A map of ValueType -> Bound Value Map constructors
 		// This is used for indirect construction when we have a UScriptStruct value type
-		TMap<UScriptStruct*, SetMapConstructorFun> SetMapInitializers;
-		mutable FRWLock SetMapInitializerLock;
+		TMap<UScriptStruct*, BoundValueMapConstructorFun> BoundValueMapInitializers;
+		mutable FRWLock BoundValueMapInitializerLock;
 
 		// A map of ValueType -> Map constructors
 		// This is used for indirect construction when we have a UScriptStruct value type
-		TMap<UScriptStruct*, MapConstructorFun> MapInitializers;
-		mutable FRWLock MapInitializerLock;
+		TMap<UScriptStruct*, UnboundValueMapConstructorFun> UnboundValueMapInitializers;
+		mutable FRWLock UnboundValueMapInitializerLock;
 
 		// A shared pointer to a map of transformer name -> transformer list (a transformer per value type)
-		FAttributeValueTransformerMapPtr TransformerMap;
+		FValueTransformerMapPtr TransformerMap;
 		mutable FRWLock TransformerMapLock;
 
 		friend class FAnimNextModuleImpl;
@@ -75,43 +75,43 @@ namespace UE::UAF
 	// Implementation
 
 	template<class ValueType>
-	inline void FValueRuntimeRegistry::RegisterSetMapInitializer()
+	inline void FValueRuntimeRegistry::RegisterBoundValueMapInitializer()
 	{
-		RegisterSetMapInitializer(ValueType::StaticStruct(), [](const FAttributeSetMap::FConstructArgs& Args)
+		RegisterBoundValueMapInitializer(ValueType::StaticStruct(), [](const FBoundValueMap::FConstructArgs& Args)
 			{
-				typename TAttributeSetMap<ValueType>::FConstructArgs TypedArgs(Args.TypedSet, Args.ValueType, Args.ReallocFun);
+				typename TBoundValueMap<ValueType>::FConstructArgs TypedArgs(Args.TypedSet, Args.ValueType, Args.ReallocFun);
 
-				return MakeAttributeSetMap<ValueType>(Args);
+				return MakeBoundValueMap<ValueType>(Args);
 			});
 	}
 
 	template<class ValueType>
-	inline void FValueRuntimeRegistry::RegisterMapInitializer()
+	inline void FValueRuntimeRegistry::RegisterUnboundValueMapInitializer()
 	{
-		RegisterMapInitializer(ValueType::StaticStruct(), [](FReallocFun ReallocFun)
+		RegisterUnboundValueMapInitializer(ValueType::StaticStruct(), [](FReallocFun ReallocFun)
 			{
-				return MakeAttributeMap<ValueType>(ReallocFun);
+				return MakeUnboundValueMap<ValueType>(ReallocFun);
 			});
 	}
 
 	template<class TransformerSpecializationType>
-	inline bool FValueRuntimeRegistry::RegisterAttributeValueTransformer()
+	inline bool FValueRuntimeRegistry::RegisterValueTransformer()
 	{
 		using TransformerCppType = typename TransformerSpecializationType::TransformerType;
 		using ValueType = typename TransformerSpecializationType::ValueType;
 
-		static_assert(std::is_base_of_v<FAttributeValueTransformer, TransformerCppType>, "Specified transformer must derive from FAttributeValueTransformer and provide the necessary API");
+		static_assert(std::is_base_of_v<FValueTransformer, TransformerCppType>, "Specified transformer must derive from FValueTransformer and provide the necessary API");
 
 		const FName TransformerName = TransformerCppType::TransformerName;
 		UScriptStruct* ValueTypeStruct = ValueType::StaticStruct();
 
 		bool bSuccess = false;
 
-		if constexpr (requires { &TransformerSpecializationType::TransformSetMap; })
+		if constexpr (requires { &TransformerSpecializationType::TransformBoundValueMap; })
 		{
 			// Signature must match our transformer
-			using FTransformSetMapFunc = typename TransformerCppType::FTransformSetMapFunc;
-			FTransformSetMapFunc SpecializedTransformFunc = &TransformerSpecializationType::TransformSetMap;
+			using FTransformBoundValueMapFunc = typename TransformerCppType::FTransformBoundValueMapFunc;
+			FTransformBoundValueMapFunc SpecializedTransformFunc = &TransformerSpecializationType::TransformBoundValueMap;
 
 #ifdef _MSC_VER
 	#pragma warning(push)
@@ -126,14 +126,14 @@ namespace UE::UAF
 	#pragma warning(pop)
 #endif
 
-			bSuccess |= RegisterAttributeValueSetMapTransformer(TransformerName, ValueTypeStruct, RawSpecializedTransformFunc);
+			bSuccess |= RegisterBouneValueMapTransformer(TransformerName, ValueTypeStruct, RawSpecializedTransformFunc);
 		}
 
-		if constexpr (requires { &TransformerSpecializationType::TransformMap; })
+		if constexpr (requires { &TransformerSpecializationType::TransformUnboundValueMap; })
 		{
 			// Signature must match our transformer
-			using FTransformMapFunc = typename TransformerCppType::FTransformMapFunc;
-			FTransformMapFunc SpecializedTransformFunc = &TransformerSpecializationType::TransformMap;
+			using FTransformUnboundValueMapFunc = typename TransformerCppType::FTransformUnboundValueMapFunc;
+			FTransformUnboundValueMapFunc SpecializedTransformFunc = &TransformerSpecializationType::TransformUnboundValueMap;
 
 #ifdef _MSC_VER
 	#pragma warning(push)
@@ -148,7 +148,7 @@ namespace UE::UAF
 	#pragma warning(pop)
 #endif
 
-			bSuccess |= RegisterAttributeValueMapTransformer(TransformerName, ValueTypeStruct, RawSpecializedTransformFunc);
+			bSuccess |= RegisterUnboundValueMapTransformer(TransformerName, ValueTypeStruct, RawSpecializedTransformFunc);
 		}
 
 		// Returns true if we registered with at least one transformer type
diff --git a/Engine/Plugins/Experimental/UAF/UAFAnimGraph/Source/UAFAnimGraph/Private/EvaluationVM/EvaluationVM.cpp b/Engine/Plugins/Experimental/UAF/UAFAnimGraph/Source/UAFAnimGraph/Private/EvaluationVM/EvaluationVM.cpp
--- a/Engine/Plugins/Experimental/UAF/UAFAnimGraph/Source/UAFAnimGraph/Private/EvaluationVM/EvaluationVM.cpp
+++ b/Engine/Plugins/Experimental/UAF/UAFAnimGraph/Source/UAFAnimGraph/Private/EvaluationVM/EvaluationVM.cpp
@@ -39,7 +39,7 @@ namespace UE::UAF
 	}
 
 	FEvaluationVM::FEvaluationVM(EEvaluationFlags InEvaluationFlags, const UAbstractSkeletonSetBinding& InSetBinding, FName InSetName, int32 InCurrentLOD)
-		: AttributeValueTransformerMap(FValueRuntimeRegistry::Get().GetTransformerMap())
+		: ValueTransformerMap(FValueRuntimeRegistry::Get().GetTransformerMap())
 		, CurrentLOD(InCurrentLOD)
 		, EvaluationFlags(InEvaluationFlags)
 	{
@@ -49,7 +49,7 @@ namespace UE::UAF
 	// REMOVE ME: Exclusively for testing that before/after yields the same results
 	FEvaluationVM::FEvaluationVM(EEvaluationFlags InEvaluationFlags, const FReferencePose& InReferencePose, const UAbstractSkeletonSetBinding* InSetBinding, FName InSetName, int32 InCurrentLOD)
 		: ReferencePose(&InReferencePose)
-		, AttributeValueTransformerMap(FValueRuntimeRegistry::Get().GetTransformerMap())
+		, ValueTransformerMap(FValueRuntimeRegistry::Get().GetTransformerMap())
 		, CurrentLOD(InCurrentLOD)
 		, EvaluationFlags(InEvaluationFlags)
 	{
@@ -168,9 +168,9 @@ namespace UE::UAF
 		return false;
 	}
 
-	const FAttributeValueTransformerMapPtr& FEvaluationVM::GetTransformerMap() const
+	const FValueTransformerMapPtr& FEvaluationVM::GetTransformerMap() const
 	{
-		return AttributeValueTransformerMap;
+		return ValueTransformerMap;
 	}
 
 	FEvaluationVMStack& FEvaluationVM::GetOrCreateStack(const FEvaluationVMStackName& StackName, uint32 TypeID)
diff --git a/Engine/Plugins/Experimental/UAF/UAFAnimGraph/Source/UAFAnimGraph/Private/EvaluationVM/Tasks/ApplyAdditiveKeyframe.cpp b/Engine/Plugins/Experimental/UAF/UAFAnimGraph/Source/UAFAnimGraph/Private/EvaluationVM/Tasks/ApplyAdditiveKeyframe.cpp
--- a/Engine/Plugins/Experimental/UAF/UAFAnimGraph/Source/UAFAnimGraph/Private/EvaluationVM/Tasks/ApplyAdditiveKeyframe.cpp
+++ b/Engine/Plugins/Experimental/UAF/UAFAnimGraph/Source/UAFAnimGraph/Private/EvaluationVM/Tasks/ApplyAdditiveKeyframe.cpp
@@ -194,12 +194,12 @@ float FAnimNextApplyAdditiveKeyframeTask::GetInterpolationAlpha(const UE::UAF::F
 	{
 		const FAttributeMappingKey MappingKey = FAttributeMappingKey::MakeFromTo<FFloatAnimationAttribute>();
 		const FAttributeCollection& AlphaCurveCollection = AlphaCurveInputIndex == 0 ? CollectionA : CollectionB;
-		if (const TAttributeSetMap<FFloatAnimationAttribute>* FloatSetMap = AlphaCurveCollection.FindSetMap<FFloatAnimationAttribute>(MappingKey))
+		if (const TBoundValueMap<FFloatAnimationAttribute>* FloatMap = AlphaCurveCollection.FindSetMap<FFloatAnimationAttribute>(MappingKey))
 		{
-			const FAttributeTypedSetPtr& TypedSet = FloatSetMap->GetTypedSet();
+			const FAttributeTypedSetPtr& TypedSet = FloatMap->GetTypedSet();
 			if (FAttributeSetIndex CurveIndex = TypedSet->FindIndex(AlphaSourceCurveName))
 			{
-				Alpha = (*FloatSetMap)[CurveIndex].Value;
+				Alpha = (*FloatMap)[CurveIndex].Value;
 			}
 			else
 			{
diff --git a/Engine/Plugins/Experimental/UAF/UAFAnimGraph/Source/UAFAnimGraph/Private/EvaluationVM/Tasks/ApplyAdditiveKeyframePerBone.cpp b/Engine/Plugins/Experimental/UAF/UAFAnimGraph/Source/UAFAnimGraph/Private/EvaluationVM/Tasks/ApplyAdditiveKeyframePerBone.cpp
--- a/Engine/Plugins/Experimental/UAF/UAFAnimGraph/Source/UAFAnimGraph/Private/EvaluationVM/Tasks/ApplyAdditiveKeyframePerBone.cpp
+++ b/Engine/Plugins/Experimental/UAF/UAFAnimGraph/Source/UAFAnimGraph/Private/EvaluationVM/Tasks/ApplyAdditiveKeyframePerBone.cpp
@@ -92,9 +92,9 @@ void FUAFApplyAdditiveKeyframePerBoneTask::Execute(UE::UAF::FEvaluationVM& VM) c
 			// Every attribute type maps to a float, our weight
 			FAttributeMappingKey MappingKey = FAttributeMappingKey::MakeFromTo(AttributeType, FloatAttributeType);
 
-			TAttributeSetMap<FFloatAnimationAttribute>* SetWeights = PerValueWeights.AddSetMap<FFloatAnimationAttribute>(MappingKey);
+			TBoundValueMap<FFloatAnimationAttribute>* SetWeights = PerValueWeights.AddSetMap<FFloatAnimationAttribute>(MappingKey);
 
-			if (const TAttributeSetMap<FBoneTransformAnimationAttribute>* BoneMap = Cast<FBoneTransformAnimationAttribute>(It.GetMap()))
+			if (const TBoundValueMap<FBoneTransformAnimationAttribute>* BoneMap = Cast<FBoneTransformAnimationAttribute>(It.GetMap()))
 			{
 				const FAttributeTypedSetPtr& BoneSet = BoneMap->GetTypedSet();
 				const int32 NumBonesInSet = BoneMap->Num();
diff --git a/Engine/Plugins/Experimental/UAF/UAFAnimGraph/Source/UAFAnimGraph/Private/EvaluationVM/Tasks/BlendKeyframes.cpp b/Engine/Plugins/Experimental/UAF/UAFAnimGraph/Source/UAFAnimGraph/Private/EvaluationVM/Tasks/BlendKeyframes.cpp
--- a/Engine/Plugins/Experimental/UAF/UAFAnimGraph/Source/UAFAnimGraph/Private/EvaluationVM/Tasks/BlendKeyframes.cpp
+++ b/Engine/Plugins/Experimental/UAF/UAFAnimGraph/Source/UAFAnimGraph/Private/EvaluationVM/Tasks/BlendKeyframes.cpp
@@ -174,12 +174,12 @@ float FAnimNextBlendTwoKeyframesTask::GetInterpolationAlpha(const UE::UAF::FAttr
 	{
 		const FAttributeMappingKey MappingKey = FAttributeMappingKey::MakeFromTo<FFloatAnimationAttribute>();
 		const FAttributeCollection& AlphaCurveCollection = AlphaCurveInputIndex == 0 ? CollectionA : CollectionB;
-		if (const TAttributeSetMap<FFloatAnimationAttribute>* FloatSetMap = AlphaCurveCollection.FindSetMap<FFloatAnimationAttribute>(MappingKey))
+		if (const TBoundValueMap<FFloatAnimationAttribute>* FloatMap = AlphaCurveCollection.FindSetMap<FFloatAnimationAttribute>(MappingKey))
 		{
-			const FAttributeTypedSetPtr& TypedSet = FloatSetMap->GetTypedSet();
+			const FAttributeTypedSetPtr& TypedSet = FloatMap->GetTypedSet();
 			if (FAttributeSetIndex CurveIndex = TypedSet->FindIndex(AlphaSourceCurveName))
 			{
-				Alpha = (*FloatSetMap)[CurveIndex].Value;
+				Alpha = (*FloatMap)[CurveIndex].Value;
 			}
 			else
 			{
diff --git a/Engine/Plugins/Experimental/UAF/UAFAnimGraph/Source/UAFAnimGraph/Private/EvaluationVM/Tasks/BlendKeyframesPerBone.cpp b/Engine/Plugins/Experimental/UAF/UAFAnimGraph/Source/UAFAnimGraph/Private/EvaluationVM/Tasks/BlendKeyframesPerBone.cpp
--- a/Engine/Plugins/Experimental/UAF/UAFAnimGraph/Source/UAFAnimGraph/Private/EvaluationVM/Tasks/BlendKeyframesPerBone.cpp
+++ b/Engine/Plugins/Experimental/UAF/UAFAnimGraph/Source/UAFAnimGraph/Private/EvaluationVM/Tasks/BlendKeyframesPerBone.cpp
@@ -41,9 +41,9 @@ namespace UE::UAF::Private
 			// Every attribute type maps to a float, our weight
 			FAttributeMappingKey MappingKey = FAttributeMappingKey::MakeFromTo(AttributeType, FloatAttributeType);
 
-			TAttributeSetMap<FFloatAnimationAttribute>* SetWeights = PerValueWeights.AddSetMap<FFloatAnimationAttribute>(MappingKey);
+			TBoundValueMap<FFloatAnimationAttribute>* SetWeights = PerValueWeights.AddSetMap<FFloatAnimationAttribute>(MappingKey);
 
-			if (const TAttributeSetMap<FBoneTransformAnimationAttribute>* BoneMap = Cast<FBoneTransformAnimationAttribute>(It.GetMap()))
+			if (const TBoundValueMap<FBoneTransformAnimationAttribute>* BoneMap = Cast<FBoneTransformAnimationAttribute>(It.GetMap()))
 			{
 				const FAttributeTypedSetPtr& BoneSet = BoneMap->GetTypedSet();
 				const int32 NumBonesInSet = BoneMap->Num();
@@ -136,9 +136,9 @@ namespace UE::UAF::Private
 			// Every attribute type maps to a float, our weight
 			FAttributeMappingKey MappingKey = FAttributeMappingKey::MakeFromTo(AttributeType, FloatAttributeType);
 
-			TAttributeSetMap<FFloatAnimationAttribute>* SetWeights = PerValueWeights.AddSetMap<FFloatAnimationAttribute>(MappingKey);
+			TBoundValueMap<FFloatAnimationAttribute>* SetWeights = PerValueWeights.AddSetMap<FFloatAnimationAttribute>(MappingKey);
 
-			if (const TAttributeSetMap<FBoneTransformAnimationAttribute>* BoneMap = Cast<FBoneTransformAnimationAttribute>(It.GetMap()))
+			if (const TBoundValueMap<FBoneTransformAnimationAttribute>* BoneMap = Cast<FBoneTransformAnimationAttribute>(It.GetMap()))
 			{
 				const FAttributeTypedSetPtr& BoneSet = BoneMap->GetTypedSet();
 				const int32 NumBonesInSet = BoneMap->Num();
@@ -176,7 +176,7 @@ namespace UE::UAF::Private
 					}
 				}
 			}
-			else if (const TAttributeSetMap<FFloatAnimationAttribute>* FloatMap = Cast<FFloatAnimationAttribute>(It.GetMap()))
+			else if (const TBoundValueMap<FFloatAnimationAttribute>* FloatMap = Cast<FFloatAnimationAttribute>(It.GetMap()))
 			{
 				SetWeights->FillWithValue(FFloatAnimationAttribute{ 0.0f });
 
diff --git a/Engine/Plugins/Experimental/UAF/UAFAnimGraph/Source/UAFAnimGraph/Private/EvaluationVM/Tasks/DecompressionTools.cpp b/Engine/Plugins/Experimental/UAF/UAFAnimGraph/Source/UAFAnimGraph/Private/EvaluationVM/Tasks/DecompressionTools.cpp
--- a/Engine/Plugins/Experimental/UAF/UAFAnimGraph/Source/UAFAnimGraph/Private/EvaluationVM/Tasks/DecompressionTools.cpp
+++ b/Engine/Plugins/Experimental/UAF/UAFAnimGraph/Source/UAFAnimGraph/Private/EvaluationVM/Tasks/DecompressionTools.cpp
@@ -98,7 +98,7 @@ bool FDecompressionTools::ShouldUseRawData(const UAnimSequence* AnimSequence, co
 namespace Private
 {
 	static void DecompressBoneTransforms(
-		TAttributeSetMap<FBoneTransformAnimationAttribute>& OutBoneTransforms,
+		TBoundValueMap<FBoneTransformAnimationAttribute>& OutBoneTransforms,
 		const FCompressedAnimSequence& CompressedData,
 		const FAnimExtractContext& ExtractionContext,
 		const FEvaluationTaskContext& VMContext,
@@ -294,7 +294,7 @@ namespace Private
 				// and we have nothing to subtract it from, we should probably use the additive identity
 				const FAttributeNamedSetPtr& NamedSet = OutBoneTransforms.GetTypedSet()->GetNamedSet();
 				const FAttributeCollectionPtr ReferenceValues = NamedSet->GetDefaultAttributeValues();
-				const TAttributeSetMap<FBoneTransformAnimationAttribute>* ReferenceBoneTransforms = ReferenceValues->FindSetMap<FBoneTransformAnimationAttribute>(OutBoneTransforms.GetMappingKey());
+				const TBoundValueMap<FBoneTransformAnimationAttribute>* ReferenceBoneTransforms = ReferenceValues->FindSetMap<FBoneTransformAnimationAttribute>(OutBoneTransforms.GetMappingKey());
 				check(ReferenceBoneTransforms != nullptr && !ReferenceBoneTransforms->IsEmpty());
 
 				const FTransform& RootRefTransform = ReferenceBoneTransforms->GetData()->Value;
@@ -444,7 +444,7 @@ namespace Private
 		const UAnimSequence* AnimSequence,
 		const FAnimExtractContext& ExtractionContext,
 		const FEvaluationTaskContext& VMContext,
-		TAttributeSetMap<FBoneTransformAnimationAttribute>& OutBoneTransforms,
+		TBoundValueMap<FBoneTransformAnimationAttribute>& OutBoneTransforms,
 		bool bForceUseRawData)
 	{
 		SCOPE_CYCLE_COUNTER(STAT_AnimSeq_GetBonePose);
@@ -609,7 +609,7 @@ void FDecompressionTools::GetAnimationPose(const UAnimSequence* AnimSequence, co
 	using namespace UE::UAF;
 
 	const FAttributeMappingKey MappingKey = FAttributeMappingKey::MakeFromTo<FBoneTransformAnimationAttribute>();
-	TAttributeSetMap<FBoneTransformAnimationAttribute>* BoneTransforms = OutCollection.FindSetMap<FBoneTransformAnimationAttribute>(MappingKey);
+	TBoundValueMap<FBoneTransformAnimationAttribute>* BoneTransforms = OutCollection.FindSetMap<FBoneTransformAnimationAttribute>(MappingKey);
 	if (BoneTransforms == nullptr || BoneTransforms->IsEmpty())
 	{
 		// Not requesting bone transforms
@@ -817,8 +817,8 @@ void FDecompressionTools::GetAnimationCurves(const UAnimSequence* AnimSequence,
 				ensureMsgf(Curves.IsSorted(), TEXT("Decompressed curves must be sorted"));
 
 				const FAttributeMappingKey MappingKey = FAttributeMappingKey::MakeFromTo<FFloatAnimationAttribute>();
-				TAttributeSetMap<FFloatAnimationAttribute>* FloatAttributes = OutCollection.FindSetMap<FFloatAnimationAttribute>(MappingKey);
-				TAttributeMap<FFloatAnimationAttribute>* FloatMap = nullptr;
+				TBoundValueMap<FFloatAnimationAttribute>* FloatAttributes = OutCollection.FindSetMap<FFloatAnimationAttribute>(MappingKey);
+				TUnboundValueMap<FFloatAnimationAttribute>* FloatMap = nullptr;
 
 				if (FloatAttributes != nullptr && !FloatAttributes->IsEmpty())
 				{
@@ -962,10 +962,10 @@ void FDecompressionTools::GetAnimationAttributes(const UAnimSequence* AnimSequen
 						{
 							// This attribute is present within the current named set
 
-							FAttributeSetMap* SetMap = OutCollection.FindSetMap(FAttributeMappingKey::MakeFromTo(AttributeType));
-							checkf(SetMap != nullptr, TEXT("This attribute exists within our named set but isn't present in the output collection"));
+							FBoundValueMap* Map = OutCollection.FindSetMap(FAttributeMappingKey::MakeFromTo(AttributeType));
+							checkf(Map != nullptr, TEXT("This attribute exists within our named set but isn't present in the output collection"));
 
-							SetMap->SetValueWithSetter(SetIndex, [&ExtractionContext, &BakedAttribute](UScriptStruct* ValueType, uint8* OutValuePtr)
+							Map->SetValueWithSetter(SetIndex, [&ExtractionContext, &BakedAttribute](UScriptStruct* ValueType, uint8* OutValuePtr)
 								{
 									BakedAttribute.Value.EvaluateToPtr(ValueType, ExtractionContext.CurrentTime, OutValuePtr);
 								});
@@ -978,7 +978,7 @@ void FDecompressionTools::GetAnimationAttributes(const UAnimSequence* AnimSequen
 					else
 					{
 						// This attribute type is unknown within our binding, it is dynamic
-						FAttributeMap* AttributeMap = OutCollection.FindOrAddMap(AttributeType);
+						FUnboundValueMap* AttributeMap = OutCollection.FindOrAddMap(AttributeType);
 
 						AttributeMap->AddWithSetter(AttributeName, [&ExtractionContext, &BakedAttribute](UScriptStruct* ValueType, uint8* OutValuePtr)
 							{
@@ -1004,10 +1004,10 @@ void FDecompressionTools::GetAnimationAttributes(const UAnimSequence* AnimSequen
 					{
 						// This attribute is present within the current named set
 
-						FAttributeSetMap* SetMap = OutCollection.FindSetMap(FAttributeMappingKey::MakeFromTo(AttributeType));
-						checkf(SetMap != nullptr, TEXT("This attribute exists within our named set but isn't present in the output collection"));
+						FBoundValueMap* Map = OutCollection.FindSetMap(FAttributeMappingKey::MakeFromTo(AttributeType));
+						checkf(Map != nullptr, TEXT("This attribute exists within our named set but isn't present in the output collection"));
 
-						SetMap->SetValueWithSetter(SetIndex, [&ExtractionContext, &BakedAttribute](UScriptStruct* ValueType, uint8* OutValuePtr)
+						Map->SetValueWithSetter(SetIndex, [&ExtractionContext, &BakedAttribute](UScriptStruct* ValueType, uint8* OutValuePtr)
 							{
 								BakedAttribute.Value.EvaluateToPtr(ValueType, ExtractionContext.CurrentTime, OutValuePtr);
 							});
@@ -1020,7 +1020,7 @@ void FDecompressionTools::GetAnimationAttributes(const UAnimSequence* AnimSequen
 				else
 				{
 					// This attribute type is unknown within our binding, it is dynamic
-					FAttributeMap* AttributeMap = OutCollection.FindOrAddMap(AttributeType);
+					FUnboundValueMap* AttributeMap = OutCollection.FindOrAddMap(AttributeType);
 
 					AttributeMap->AddWithSetter(AttributeName, [&ExtractionContext, &BakedAttribute](UScriptStruct* ValueType, uint8* OutValuePtr)
 						{
diff --git a/Engine/Plugins/Experimental/UAF/UAFAnimGraph/Source/UAFAnimGraph/Private/EvaluationVM/Tasks/PushAnimSequenceKeyframe.cpp b/Engine/Plugins/Experimental/UAF/UAFAnimGraph/Source/UAFAnimGraph/Private/EvaluationVM/Tasks/PushAnimSequenceKeyframe.cpp
--- a/Engine/Plugins/Experimental/UAF/UAFAnimGraph/Source/UAFAnimGraph/Private/EvaluationVM/Tasks/PushAnimSequenceKeyframe.cpp
+++ b/Engine/Plugins/Experimental/UAF/UAFAnimGraph/Source/UAFAnimGraph/Private/EvaluationVM/Tasks/PushAnimSequenceKeyframe.cpp
@@ -76,11 +76,11 @@ void FAnimNextAnimSequenceKeyframeTask::Execute(UE::UAF::FEvaluationVM& VM) cons
 						// This attribute is present within the current named set
 
 						FAttributeMappingKey MappingKey = FAttributeMappingKey::MakeFromTo<FTransformAnimationAttribute>();
-						TAttributeSetMap<FTransformAnimationAttribute>* SetMap = Collection.FindSetMap<FTransformAnimationAttribute>(MappingKey);
-						checkf(SetMap != nullptr, TEXT("This attribute exists within our named set but isn't present in the output collection"));
+						TBoundValueMap<FTransformAnimationAttribute>* Map = Collection.FindSetMap<FTransformAnimationAttribute>(MappingKey);
+						checkf(Map != nullptr, TEXT("This attribute exists within our named set but isn't present in the output collection"));
 
 						const FTransform RootMotionTransform = AnimSequence->ExtractRootMotion(ExtractionContext);
-						(*SetMap)[SetIndex].Value = RootMotionTransform;
+						(*Map)[SetIndex].Value = RootMotionTransform;
 					}
 				}
 			}
diff --git a/Engine/Plugins/Experimental/UAF/UAFAnimGraph/Source/UAFAnimGraph/Private/TraitInterfaces/IEvaluate.cpp b/Engine/Plugins/Experimental/UAF/UAFAnimGraph/Source/UAFAnimGraph/Private/TraitInterfaces/IEvaluate.cpp
--- a/Engine/Plugins/Experimental/UAF/UAFAnimGraph/Source/UAFAnimGraph/Private/TraitInterfaces/IEvaluate.cpp
+++ b/Engine/Plugins/Experimental/UAF/UAFAnimGraph/Source/UAFAnimGraph/Private/TraitInterfaces/IEvaluate.cpp
@@ -269,10 +269,10 @@ namespace UE::UAF
 						// Set maps
 						for (auto It = Collection->CreateConstSetMapIterator(); It; ++It)
 						{
-							const FAttributeSetMap* SetMap = It.GetMap();
-							const FAttributeTypedSetPtr& TypedSet = SetMap->GetTypedSet();
+							const FBoundValueMap* Map = It.GetMap();
+							const FAttributeTypedSetPtr& TypedSet = Map->GetTypedSet();
 
-							if (const TAttributeSetMap<FBoneTransformAnimationAttribute>* BoneTransforms = Cast<FBoneTransformAnimationAttribute>(SetMap))
+							if (const TBoundValueMap<FBoneTransformAnimationAttribute>* BoneTransforms = Cast<FBoneTransformAnimationAttribute>(Map))
 							{
 								const int32 NumPoseLODBones = OutputPose.LODPose.GetNumBones();
 
@@ -294,7 +294,7 @@ namespace UE::UAF
 									}
 								}
 							}
-							else if (const TAttributeSetMap<FFloatAnimationAttribute>* Floats = Cast<FFloatAnimationAttribute>(SetMap))
+							else if (const TBoundValueMap<FFloatAnimationAttribute>* Floats = Cast<FFloatAnimationAttribute>(Map))
 							{
 								const auto GetNameFromIndex = [&TypedSet](int32 CurveIndex)
 									{
@@ -312,8 +312,8 @@ namespace UE::UAF
 							else
 							{
 								// Everything else
-								const UScriptStruct* ValueType = SetMap->GetValueType();
-								const int32 NumAttributes = SetMap->Num();
+								const UScriptStruct* ValueType = Map->GetValueType();
+								const int32 NumAttributes = Map->Num();
 
 								for (FAttributeSetIndex SetIndex(0); SetIndex < NumAttributes; ++SetIndex)
 								{
@@ -321,7 +321,7 @@ namespace UE::UAF
 
 									uint8* AttributePtr = OutputPose.Attributes.FindOrAdd(ValueType, AttributeId);
 
-									SetMap->GetValueWithGetter(SetIndex, [AttributePtr](UScriptStruct* ValueType_, const uint8* ValuePtr)
+									Map->GetValueWithGetter(SetIndex, [AttributePtr](UScriptStruct* ValueType_, const uint8* ValuePtr)
 										{
 											ValueType_->CopyScriptStruct(AttributePtr, ValuePtr);
 										});
@@ -332,7 +332,7 @@ namespace UE::UAF
 						// Name/value maps
 						for (auto It = Collection->CreateConstMapIterator(); It; ++It)
 						{
-							const FAttributeMap* Map = It.GetMap();
+							const FUnboundValueMap* Map = It.GetMap();
 							const UScriptStruct* ValueType = Map->GetValueType();
 							const int32 NumValues = Map->Num();
 
diff --git a/Engine/Plugins/Experimental/UAF/UAFAnimGraph/Source/UAFAnimGraph/Private/Traits/ModifyCurveTrait.cpp b/Engine/Plugins/Experimental/UAF/UAFAnimGraph/Source/UAFAnimGraph/Private/Traits/ModifyCurveTrait.cpp
--- a/Engine/Plugins/Experimental/UAF/UAFAnimGraph/Source/UAFAnimGraph/Private/Traits/ModifyCurveTrait.cpp
+++ b/Engine/Plugins/Experimental/UAF/UAFAnimGraph/Source/UAFAnimGraph/Private/Traits/ModifyCurveTrait.cpp
@@ -60,7 +60,7 @@ void FModifyCurveTask::Execute(UE::UAF::FEvaluationVM& VM) const
 			{
 				FAttributeMappingKey MappingKey = FAttributeMappingKey::MakeFromTo<FFloatAnimationAttribute>();
 				
-				if (TAttributeMap<FFloatAnimationAttribute>* FloatMap = (*Collection)->FindMap<FFloatAnimationAttribute>())
+				if (TUnboundValueMap<FFloatAnimationAttribute>* FloatMap = (*Collection)->FindMap<FFloatAnimationAttribute>())
 				{
 					for (const FModifyCurveParameters& CurveParameters : SharedData->ModifyCurveParameters)
 					{
@@ -73,15 +73,15 @@ void FModifyCurveTask::Execute(UE::UAF::FEvaluationVM& VM) const
 					}
 				}
 
-				if (TAttributeSetMap<FFloatAnimationAttribute>*	FloatSetMap = (*Collection)->FindSetMap<FFloatAnimationAttribute>(MappingKey))
+				if (TBoundValueMap<FFloatAnimationAttribute>* FloatMap = (*Collection)->FindSetMap<FFloatAnimationAttribute>(MappingKey))
 				{
 					for (const FModifyCurveParameters& CurveParameters : SharedData->ModifyCurveParameters)
 					{
 						if (const FAttributeSetIndex RootMotionAttributeIndex = FloatTypedSet->FindIndex(CurveParameters.CurveName))
 						{
-							const float CurrentValue = (*FloatSetMap)[RootMotionAttributeIndex].Value;
+							const float CurrentValue = (*FloatMap)[RootMotionAttributeIndex].Value;
 							const float UpdatedValue = ProcessCurveOperation(CurrentValue, CurveParameters.CurveValue, InstanceData->Alpha, SharedData->ApplyMode);
-							(*FloatSetMap)[RootMotionAttributeIndex].Value = UpdatedValue;
+							(*FloatMap)[RootMotionAttributeIndex].Value = UpdatedValue;
 						}
 					}
 				}
diff --git a/Engine/Plugins/Experimental/UAF/UAFAnimGraph/Source/UAFAnimGraph/Public/EvaluationVM/EvaluationVM.h b/Engine/Plugins/Experimental/UAF/UAFAnimGraph/Source/UAFAnimGraph/Public/EvaluationVM/EvaluationVM.h
--- a/Engine/Plugins/Experimental/UAF/UAFAnimGraph/Source/UAFAnimGraph/Public/EvaluationVM/EvaluationVM.h
+++ b/Engine/Plugins/Experimental/UAF/UAFAnimGraph/Source/UAFAnimGraph/Public/EvaluationVM/EvaluationVM.h
@@ -218,8 +218,8 @@ namespace UE::UAF
 		// Returns the current LOD
 		[[nodiscard]] int32 GetCurrentLOD() const { return CurrentLOD; }
 
-		// Returns the attribute value transformer map
-		[[nodiscard]] UE_API const FAttributeValueTransformerMapPtr& GetTransformerMap() const;
+		// Returns the value transformer map
+		[[nodiscard]] UE_API const FValueTransformerMapPtr& GetTransformerMap() const;
 
 	private:
 		// Disallow copy
@@ -244,8 +244,8 @@ namespace UE::UAF
 		// The internal context that controls what state we are producing
 		TArray<FEvaluationTaskContext> EvaluationContextStack;
 
-		// Shared map of attribute value transformers
-		FAttributeValueTransformerMapPtr AttributeValueTransformerMap;
+		// Shared map of value transformers
+		FValueTransformerMapPtr ValueTransformerMap;
 
 		// Current LOD we are evaluating at
 		int32 CurrentLOD = 0;
diff --git a/Engine/Plugins/Experimental/UAF/UAFTests/Source/UAFTestSuite/Private/UAF/ValueRuntime/AttributeCollectionTest.cpp b/Engine/Plugins/Experimental/UAF/UAFTests/Source/UAFTestSuite/Private/UAF/ValueRuntime/AttributeCollectionTest.cpp
--- a/Engine/Plugins/Experimental/UAF/UAFTests/Source/UAFTestSuite/Private/UAF/ValueRuntime/AttributeCollectionTest.cpp
+++ b/Engine/Plugins/Experimental/UAF/UAFTests/Source/UAFTestSuite/Private/UAF/ValueRuntime/AttributeCollectionTest.cpp
@@ -14,14 +14,14 @@
 #include "UAF/Attributes/AttributeBindingData.h"
 #include "UAF/Attributes/EngineAttributes.h"
 #include "UAF/ValueRuntime/AttributeCollection.h"
-#include "UAF/ValueRuntime/AttributeSetMap.h"
+#include "UAF/ValueRuntime/BoundValueMap.h"
 #include "UAF/ValueRuntime/IteratorTestUtils.h"
 #include "UObject/UObjectGlobals.h"
 #include "UObject/Package.h"
 
 namespace UE::UAF::Tests
 {
-	IMPLEMENT_SIMPLE_AUTOMATION_TEST(FAttributeCollectionTest, "Animation.UAF.Attributes.AttributeCollection", EAutomationTestFlags::EditorContext | EAutomationTestFlags::EngineFilter)
+	IMPLEMENT_SIMPLE_AUTOMATION_TEST(FAttributeCollectionTest, "Animation.UAF.ValueRuntime.AttributeCollection", EAutomationTestFlags::EditorContext | EAutomationTestFlags::EngineFilter)
 
 	bool FAttributeCollectionTest::RunTest(const FString& InParameters)
 	{
@@ -96,11 +96,11 @@ namespace UE::UAF::Tests
 			AddErrorIfFalse(Collection.GetNamedSet() == FullBodySet, TEXT("Unexpected named set"));
 
 			FAttributeMappingKey FloatToFloatKey = FAttributeMappingKey::MakeFromTo<FFloatAnimationAttribute, FFloatAnimationAttribute>();
-			TAttributeSetMap<FFloatAnimationAttribute>* FloatToFloatMap = Collection.AddSetMap<FFloatAnimationAttribute>(FloatToFloatKey);
+			TBoundValueMap<FFloatAnimationAttribute>* FloatToFloatMap = Collection.AddSetMap<FFloatAnimationAttribute>(FloatToFloatKey);
 			UE_RETURN_ON_ERROR(FloatToFloatMap != nullptr, TEXT("Float->Float map should be valid"));
 
 			FAttributeMappingKey FloatToIntegerKey = FAttributeMappingKey::MakeFromTo<FFloatAnimationAttribute, FIntegerAnimationAttribute>();
-			TAttributeSetMap<FIntegerAnimationAttribute>* FloatToIntegerMap = MakeAttributeSetMap<FIntegerAnimationAttribute>(FAttributeSetMap::FConstructArgs(FullBodySet->FindTypedSet<FFloatAnimationAttribute>(), FIntegerAnimationAttribute::StaticStruct(), Collection.GetAllocator()));
+			TBoundValueMap<FIntegerAnimationAttribute>* FloatToIntegerMap = MakeBoundValueMap<FIntegerAnimationAttribute>(FBoundValueMap::FConstructArgs(FullBodySet->FindTypedSet<FFloatAnimationAttribute>(), FIntegerAnimationAttribute::StaticStruct(), Collection.GetAllocator()));
 			UE_RETURN_ON_ERROR(FloatToIntegerMap != nullptr, TEXT("Float->Integer map should be valid"));
 			bool bAdded = Collection.AddSetMap(FloatToIntegerMap);
 			UE_RETURN_ON_ERROR(bAdded, TEXT("Float->Integer map should be added"));
@@ -110,24 +110,24 @@ namespace UE::UAF::Tests
 			AddErrorIfFalse(!bAdded, TEXT("Float->Integer map should not be added"));
 
 			FAttributeMappingKey FloatToStringKey = FAttributeMappingKey::MakeFromTo<FFloatAnimationAttribute, FStringAnimationAttribute>();
-			TAttributeSetMap<FStringAnimationAttribute>* FloatToStringMap = Collection.AddSetMap<FStringAnimationAttribute>(FloatToStringKey);
+			TBoundValueMap<FStringAnimationAttribute>* FloatToStringMap = Collection.AddSetMap<FStringAnimationAttribute>(FloatToStringKey);
 			UE_RETURN_ON_ERROR(FloatToStringMap != nullptr, TEXT("Float->String map should be valid"));
 
 			FAttributeMappingKey TransformToTransformKey = FAttributeMappingKey::MakeFromTo<FBoneTransformAnimationAttribute, FBoneTransformAnimationAttribute>();
-			TAttributeSetMap<FBoneTransformAnimationAttribute>* TransformToTransformMap = Collection.AddSetMap<FBoneTransformAnimationAttribute>(TransformToTransformKey);
+			TBoundValueMap<FBoneTransformAnimationAttribute>* TransformToTransformMap = Collection.AddSetMap<FBoneTransformAnimationAttribute>(TransformToTransformKey);
 			UE_RETURN_ON_ERROR(TransformToTransformMap != nullptr, TEXT("Transform->Transform map should be valid"));
 
 			FAttributeMappingKey TransformToFloatKey = FAttributeMappingKey::MakeFromTo<FBoneTransformAnimationAttribute, FFloatAnimationAttribute>();
-			TAttributeSetMap<FFloatAnimationAttribute>* TransformToFloatMap = Cast<FFloatAnimationAttribute>(Collection.AddSetMap(TransformToFloatKey));
+			TBoundValueMap<FFloatAnimationAttribute>* TransformToFloatMap = Cast<FFloatAnimationAttribute>(Collection.AddSetMap(TransformToFloatKey));
 			UE_RETURN_ON_ERROR(TransformToFloatMap != nullptr, TEXT("Transform->Bone map should be valid"));
 
-			TAttributeMap<FFloatAnimationAttribute>* FloatMap = Collection.AddMap<FFloatAnimationAttribute>();
+			TUnboundValueMap<FFloatAnimationAttribute>* FloatMap = Collection.AddMap<FFloatAnimationAttribute>();
 			UE_RETURN_ON_ERROR(FloatMap != nullptr, TEXT("Float map should be valid"));
 
-			TAttributeMap<FBoneTransformAnimationAttribute>* BoneMap = Cast<FBoneTransformAnimationAttribute>(Collection.AddMap(FBoneTransformAnimationAttribute::StaticStruct()));
+			TUnboundValueMap<FBoneTransformAnimationAttribute>* BoneMap = Cast<FBoneTransformAnimationAttribute>(Collection.AddMap(FBoneTransformAnimationAttribute::StaticStruct()));
 			UE_RETURN_ON_ERROR(BoneMap != nullptr, TEXT("Bone map should be valid"));
 
-			TAttributeMap<FIntegerAnimationAttribute>* IntMap = MakeAttributeMap<FIntegerAnimationAttribute>(Collection.GetAllocator());
+			TUnboundValueMap<FIntegerAnimationAttribute>* IntMap = MakeUnboundValueMap<FIntegerAnimationAttribute>(Collection.GetAllocator());
 			UE_RETURN_ON_ERROR(IntMap != nullptr, TEXT("Integer map should be valid"));
 			bAdded = Collection.AddMap(IntMap);
 			UE_RETURN_ON_ERROR(bAdded, TEXT("Integer map should be added"));
@@ -212,14 +212,14 @@ namespace UE::UAF::Tests
 
 				for (auto It = Collection.CreateConstSetMapIterator(); It; ++It)
 				{
-					FAttributeSetMap* SetMapCopy = It.GetMap()->Duplicate(CollectionCopy.GetAllocator());
+					FBoundValueMap* SetMapCopy = It.GetMap()->Duplicate(CollectionCopy.GetAllocator());
 					const bool bAppended = CollectionCopy.AppendSetMap(SetMapCopy);
 					AddErrorIfFalse(bAppended, TEXT("Set map should be appended"));
 				}
 
 				for (auto It = Collection.CreateConstMapIterator(); It; ++It)
 				{
-					FAttributeMap* MapCopy = It.GetMap()->Duplicate(CollectionCopy.GetAllocator());
+					FUnboundValueMap* MapCopy = It.GetMap()->Duplicate(CollectionCopy.GetAllocator());
 					const bool bAppended = CollectionCopy.AppendMap(MapCopy);
 					AddErrorIfFalse(bAppended, TEXT("Map should be appended"));
 				}
@@ -228,25 +228,25 @@ namespace UE::UAF::Tests
 				AddErrorIfFalse(CollectionCopy.NumSetMaps() == Collection.NumSetMaps(), TEXT("Unexpected set map count"));
 				AddErrorIfFalse(CollectionCopy.NumMaps() == Collection.NumMaps(), TEXT("Unexpected map count"));
 
-				TAttributeSetMap<FFloatAnimationAttribute>* FloatToFloatMapCopy = CollectionCopy.FindSetMap<FFloatAnimationAttribute>(FloatToFloatKey);
+				TBoundValueMap<FFloatAnimationAttribute>* FloatToFloatMapCopy = CollectionCopy.FindSetMap<FFloatAnimationAttribute>(FloatToFloatKey);
 				UE_RETURN_ON_ERROR(FloatToFloatMapCopy != nullptr, TEXT("Float->Float map should be valid"));
 
-				TAttributeSetMap<FIntegerAnimationAttribute>* FloatToIntegerMapCopy = CollectionCopy.FindSetMap<FIntegerAnimationAttribute>(FloatToIntegerKey);
+				TBoundValueMap<FIntegerAnimationAttribute>* FloatToIntegerMapCopy = CollectionCopy.FindSetMap<FIntegerAnimationAttribute>(FloatToIntegerKey);
 				UE_RETURN_ON_ERROR(FloatToIntegerMapCopy != nullptr, TEXT("Float->Integer map should be valid"));
 
-				TAttributeSetMap<FStringAnimationAttribute>* FloatToStringMapCopy = CollectionCopy.FindSetMap<FStringAnimationAttribute>(FloatToStringKey);
+				TBoundValueMap<FStringAnimationAttribute>* FloatToStringMapCopy = CollectionCopy.FindSetMap<FStringAnimationAttribute>(FloatToStringKey);
 				UE_RETURN_ON_ERROR(FloatToStringMapCopy != nullptr, TEXT("Float->String map should be valid"));
 
-				TAttributeSetMap<FBoneTransformAnimationAttribute>* TransformToTransformMapCopy = CollectionCopy.FindSetMap<FBoneTransformAnimationAttribute>(TransformToTransformKey);
+				TBoundValueMap<FBoneTransformAnimationAttribute>* TransformToTransformMapCopy = CollectionCopy.FindSetMap<FBoneTransformAnimationAttribute>(TransformToTransformKey);
 				UE_RETURN_ON_ERROR(TransformToTransformMapCopy != nullptr, TEXT("Transform->Transform map should be valid"));
 
-				TAttributeSetMap<FFloatAnimationAttribute>* TransformToFloatMapCopy = Cast<FFloatAnimationAttribute>(CollectionCopy.FindSetMap(TransformToFloatKey));
+				TBoundValueMap<FFloatAnimationAttribute>* TransformToFloatMapCopy = Cast<FFloatAnimationAttribute>(CollectionCopy.FindSetMap(TransformToFloatKey));
 				UE_RETURN_ON_ERROR(TransformToFloatMapCopy != nullptr, TEXT("Transform->Bone map should be valid"));
 
-				TAttributeMap<FFloatAnimationAttribute>* FloatMapCopy = CollectionCopy.FindMap<FFloatAnimationAttribute>();
+				TUnboundValueMap<FFloatAnimationAttribute>* FloatMapCopy = CollectionCopy.FindMap<FFloatAnimationAttribute>();
 				UE_RETURN_ON_ERROR(FloatMapCopy != nullptr, TEXT("Float map should be valid"));
 
-				TAttributeMap<FBoneTransformAnimationAttribute>* BoneMapCopy = Cast<FBoneTransformAnimationAttribute>(CollectionCopy.FindMap(FBoneTransformAnimationAttribute::StaticStruct()));
+				TUnboundValueMap<FBoneTransformAnimationAttribute>* BoneMapCopy = Cast<FBoneTransformAnimationAttribute>(CollectionCopy.FindMap(FBoneTransformAnimationAttribute::StaticStruct()));
 				UE_RETURN_ON_ERROR(BoneMapCopy != nullptr, TEXT("Bone map should be valid"));
 
 				AddErrorIfFalse(FloatToFloatMap != FloatToFloatMapCopy, TEXT("Set map instances should be duplicated"));
@@ -272,25 +272,25 @@ namespace UE::UAF::Tests
 				AddErrorIfFalse(CollectionCopy.NumSetMaps() == Collection.NumSetMaps(), TEXT("Unexpected set map count"));
 				AddErrorIfFalse(CollectionCopy.NumMaps() == Collection.NumMaps(), TEXT("Unexpected map count"));
 
-				TAttributeSetMap<FFloatAnimationAttribute>* FloatToFloatMapCopy = CollectionCopy.FindSetMap<FFloatAnimationAttribute>(FloatToFloatKey);
+				TBoundValueMap<FFloatAnimationAttribute>* FloatToFloatMapCopy = CollectionCopy.FindSetMap<FFloatAnimationAttribute>(FloatToFloatKey);
 				UE_RETURN_ON_ERROR(FloatToFloatMapCopy != nullptr, TEXT("Float->Float map should be valid"));
 
-				TAttributeSetMap<FIntegerAnimationAttribute>* FloatToIntegerMapCopy = CollectionCopy.FindSetMap<FIntegerAnimationAttribute>(FloatToIntegerKey);
+				TBoundValueMap<FIntegerAnimationAttribute>* FloatToIntegerMapCopy = CollectionCopy.FindSetMap<FIntegerAnimationAttribute>(FloatToIntegerKey);
 				UE_RETURN_ON_ERROR(FloatToIntegerMapCopy != nullptr, TEXT("Float->Integer map should be valid"));
 
-				TAttributeSetMap<FStringAnimationAttribute>* FloatToStringMapCopy = CollectionCopy.FindSetMap<FStringAnimationAttribute>(FloatToStringKey);
+				TBoundValueMap<FStringAnimationAttribute>* FloatToStringMapCopy = CollectionCopy.FindSetMap<FStringAnimationAttribute>(FloatToStringKey);
 				UE_RETURN_ON_ERROR(FloatToStringMapCopy != nullptr, TEXT("Float->String map should be valid"));
 
-				TAttributeSetMap<FBoneTransformAnimationAttribute>* TransformToTransformMapCopy = CollectionCopy.FindSetMap<FBoneTransformAnimationAttribute>(TransformToTransformKey);
+				TBoundValueMap<FBoneTransformAnimationAttribute>* TransformToTransformMapCopy = CollectionCopy.FindSetMap<FBoneTransformAnimationAttribute>(TransformToTransformKey);
 				UE_RETURN_ON_ERROR(TransformToTransformMapCopy != nullptr, TEXT("Transform->Transform map should be valid"));
 
-				TAttributeSetMap<FFloatAnimationAttribute>* TransformToFloatMapCopy = Cast<FFloatAnimationAttribute>(CollectionCopy.FindSetMap(TransformToFloatKey));
+				TBoundValueMap<FFloatAnimationAttribute>* TransformToFloatMapCopy = Cast<FFloatAnimationAttribute>(CollectionCopy.FindSetMap(TransformToFloatKey));
 				UE_RETURN_ON_ERROR(TransformToFloatMapCopy != nullptr, TEXT("Transform->Bone map should be valid"));
 
-				TAttributeMap<FFloatAnimationAttribute>* FloatMapCopy = CollectionCopy.FindMap<FFloatAnimationAttribute>();
+				TUnboundValueMap<FFloatAnimationAttribute>* FloatMapCopy = CollectionCopy.FindMap<FFloatAnimationAttribute>();
 				UE_RETURN_ON_ERROR(FloatMapCopy != nullptr, TEXT("Float map should be valid"));
 
-				TAttributeMap<FBoneTransformAnimationAttribute>* BoneMapCopy = Cast<FBoneTransformAnimationAttribute>(CollectionCopy.FindMap(FBoneTransformAnimationAttribute::StaticStruct()));
+				TUnboundValueMap<FBoneTransformAnimationAttribute>* BoneMapCopy = Cast<FBoneTransformAnimationAttribute>(CollectionCopy.FindMap(FBoneTransformAnimationAttribute::StaticStruct()));
 				UE_RETURN_ON_ERROR(BoneMapCopy != nullptr, TEXT("Bone map should be valid"));
 
 				AddErrorIfFalse(FloatToFloatMap != FloatToFloatMapCopy, TEXT("Set map instances should be duplicated"));
@@ -316,25 +316,25 @@ namespace UE::UAF::Tests
 					AddErrorIfFalse(CollectionMoved.NumSetMaps() == Collection.NumSetMaps(), TEXT("Unexpected set map count"));
 					AddErrorIfFalse(CollectionMoved.NumMaps() == Collection.NumMaps(), TEXT("Unexpected map count"));
 
-					TAttributeSetMap<FFloatAnimationAttribute>* FloatToFloatMapMoved = CollectionMoved.FindSetMap<FFloatAnimationAttribute>(FloatToFloatKey);
+					TBoundValueMap<FFloatAnimationAttribute>* FloatToFloatMapMoved = CollectionMoved.FindSetMap<FFloatAnimationAttribute>(FloatToFloatKey);
 					UE_RETURN_ON_ERROR(FloatToFloatMapMoved != nullptr, TEXT("Float->Float map should be valid"));
 
-					TAttributeSetMap<FIntegerAnimationAttribute>* FloatToIntegerMapMoved = CollectionMoved.FindSetMap<FIntegerAnimationAttribute>(FloatToIntegerKey);
+					TBoundValueMap<FIntegerAnimationAttribute>* FloatToIntegerMapMoved = CollectionMoved.FindSetMap<FIntegerAnimationAttribute>(FloatToIntegerKey);
 					UE_RETURN_ON_ERROR(FloatToIntegerMapMoved != nullptr, TEXT("Float->Integer map should be valid"));
 
-					TAttributeSetMap<FStringAnimationAttribute>* FloatToStringMapMoved = CollectionMoved.FindSetMap<FStringAnimationAttribute>(FloatToStringKey);
+					TBoundValueMap<FStringAnimationAttribute>* FloatToStringMapMoved = CollectionMoved.FindSetMap<FStringAnimationAttribute>(FloatToStringKey);
 					UE_RETURN_ON_ERROR(FloatToStringMapMoved != nullptr, TEXT("Float->String map should be valid"));
 
-					TAttributeSetMap<FBoneTransformAnimationAttribute>* TransformToTransformMapMoved = CollectionMoved.FindSetMap<FBoneTransformAnimationAttribute>(TransformToTransformKey);
+					TBoundValueMap<FBoneTransformAnimationAttribute>* TransformToTransformMapMoved = CollectionMoved.FindSetMap<FBoneTransformAnimationAttribute>(TransformToTransformKey);
 					UE_RETURN_ON_ERROR(TransformToTransformMapMoved != nullptr, TEXT("Transform->Transform map should be valid"));
 
-					TAttributeSetMap<FFloatAnimationAttribute>* TransformToFloatMapMoved = Cast<FFloatAnimationAttribute>(CollectionMoved.FindSetMap(TransformToFloatKey));
+					TBoundValueMap<FFloatAnimationAttribute>* TransformToFloatMapMoved = Cast<FFloatAnimationAttribute>(CollectionMoved.FindSetMap(TransformToFloatKey));
 					UE_RETURN_ON_ERROR(TransformToFloatMapMoved != nullptr, TEXT("Transform->Bone map should be valid"));
 
-					TAttributeMap<FFloatAnimationAttribute>* FloatMapMoved = CollectionMoved.FindMap<FFloatAnimationAttribute>();
+					TUnboundValueMap<FFloatAnimationAttribute>* FloatMapMoved = CollectionMoved.FindMap<FFloatAnimationAttribute>();
 					UE_RETURN_ON_ERROR(FloatMapMoved != nullptr, TEXT("Float map should be valid"));
 
-					TAttributeMap<FBoneTransformAnimationAttribute>* BoneMapMoved = Cast<FBoneTransformAnimationAttribute>(CollectionMoved.FindMap(FBoneTransformAnimationAttribute::StaticStruct()));
+					TUnboundValueMap<FBoneTransformAnimationAttribute>* BoneMapMoved = Cast<FBoneTransformAnimationAttribute>(CollectionMoved.FindMap(FBoneTransformAnimationAttribute::StaticStruct()));
 					UE_RETURN_ON_ERROR(BoneMapMoved != nullptr, TEXT("Bone map should be valid"));
 
 					AddErrorIfFalse(FloatToFloatMapMoved == FloatToFloatMapCopy, TEXT("Set map instances should be moved"));
@@ -356,25 +356,25 @@ namespace UE::UAF::Tests
 				AddErrorIfFalse(CollectionCopy.NumSetMaps() == Collection.NumSetMaps(), TEXT("Unexpected set map count"));
 				AddErrorIfFalse(CollectionCopy.NumMaps() == Collection.NumMaps(), TEXT("Unexpected map count"));
 
-				TAttributeSetMap<FFloatAnimationAttribute>* FloatToFloatMapCopy = CollectionCopy.FindSetMap<FFloatAnimationAttribute>(FloatToFloatKey);
+				TBoundValueMap<FFloatAnimationAttribute>* FloatToFloatMapCopy = CollectionCopy.FindSetMap<FFloatAnimationAttribute>(FloatToFloatKey);
 				UE_RETURN_ON_ERROR(FloatToFloatMapCopy != nullptr, TEXT("Float->Float map should be valid"));
 
-				TAttributeSetMap<FIntegerAnimationAttribute>* FloatToIntegerMapCopy = CollectionCopy.FindSetMap<FIntegerAnimationAttribute>(FloatToIntegerKey);
+				TBoundValueMap<FIntegerAnimationAttribute>* FloatToIntegerMapCopy = CollectionCopy.FindSetMap<FIntegerAnimationAttribute>(FloatToIntegerKey);
 				UE_RETURN_ON_ERROR(FloatToIntegerMapCopy != nullptr, TEXT("Float->Integer map should be valid"));
 
-				TAttributeSetMap<FStringAnimationAttribute>* FloatToStringMapCopy = CollectionCopy.FindSetMap<FStringAnimationAttribute>(FloatToStringKey);
+				TBoundValueMap<FStringAnimationAttribute>* FloatToStringMapCopy = CollectionCopy.FindSetMap<FStringAnimationAttribute>(FloatToStringKey);
 				UE_RETURN_ON_ERROR(FloatToStringMapCopy != nullptr, TEXT("Float->String map should be valid"));
 
-				TAttributeSetMap<FBoneTransformAnimationAttribute>* TransformToTransformMapCopy = CollectionCopy.FindSetMap<FBoneTransformAnimationAttribute>(TransformToTransformKey);
+				TBoundValueMap<FBoneTransformAnimationAttribute>* TransformToTransformMapCopy = CollectionCopy.FindSetMap<FBoneTransformAnimationAttribute>(TransformToTransformKey);
 				UE_RETURN_ON_ERROR(TransformToTransformMapCopy != nullptr, TEXT("Transform->Transform map should be valid"));
 
-				TAttributeSetMap<FFloatAnimationAttribute>* TransformToFloatMapCopy = Cast<FFloatAnimationAttribute>(CollectionCopy.FindSetMap(TransformToFloatKey));
+				TBoundValueMap<FFloatAnimationAttribute>* TransformToFloatMapCopy = Cast<FFloatAnimationAttribute>(CollectionCopy.FindSetMap(TransformToFloatKey));
 				UE_RETURN_ON_ERROR(TransformToFloatMapCopy != nullptr, TEXT("Transform->Bone map should be valid"));
 
-				TAttributeMap<FFloatAnimationAttribute>* FloatMapCopy = CollectionCopy.FindMap<FFloatAnimationAttribute>();
+				TUnboundValueMap<FFloatAnimationAttribute>* FloatMapCopy = CollectionCopy.FindMap<FFloatAnimationAttribute>();
 				UE_RETURN_ON_ERROR(FloatMapCopy != nullptr, TEXT("Float map should be valid"));
 
-				TAttributeMap<FBoneTransformAnimationAttribute>* BoneMapCopy = Cast<FBoneTransformAnimationAttribute>(CollectionCopy.FindMap(FBoneTransformAnimationAttribute::StaticStruct()));
+				TUnboundValueMap<FBoneTransformAnimationAttribute>* BoneMapCopy = Cast<FBoneTransformAnimationAttribute>(CollectionCopy.FindMap(FBoneTransformAnimationAttribute::StaticStruct()));
 				UE_RETURN_ON_ERROR(BoneMapCopy != nullptr, TEXT("Bone map should be valid"));
 
 				AddErrorIfFalse(FloatToFloatMap != FloatToFloatMapCopy, TEXT("Set map instances should be duplicated"));
@@ -401,25 +401,25 @@ namespace UE::UAF::Tests
 					AddErrorIfFalse(CollectionMoved.NumSetMaps() == Collection.NumSetMaps(), TEXT("Unexpected set map count"));
 					AddErrorIfFalse(CollectionMoved.NumMaps() == Collection.NumMaps(), TEXT("Unexpected map count"));
 
-					TAttributeSetMap<FFloatAnimationAttribute>* FloatToFloatMapMoved = CollectionMoved.FindSetMap<FFloatAnimationAttribute>(FloatToFloatKey);
+					TBoundValueMap<FFloatAnimationAttribute>* FloatToFloatMapMoved = CollectionMoved.FindSetMap<FFloatAnimationAttribute>(FloatToFloatKey);
 					UE_RETURN_ON_ERROR(FloatToFloatMapMoved != nullptr, TEXT("Float->Float map should be valid"));
 
-					TAttributeSetMap<FIntegerAnimationAttribute>* FloatToIntegerMapMoved = CollectionMoved.FindSetMap<FIntegerAnimationAttribute>(FloatToIntegerKey);
+					TBoundValueMap<FIntegerAnimationAttribute>* FloatToIntegerMapMoved = CollectionMoved.FindSetMap<FIntegerAnimationAttribute>(FloatToIntegerKey);
 					UE_RETURN_ON_ERROR(FloatToIntegerMapMoved != nullptr, TEXT("Float->Integer map should be valid"));
 
-					TAttributeSetMap<FStringAnimationAttribute>* FloatToStringMapMoved = CollectionMoved.FindSetMap<FStringAnimationAttribute>(FloatToStringKey);
+					TBoundValueMap<FStringAnimationAttribute>* FloatToStringMapMoved = CollectionMoved.FindSetMap<FStringAnimationAttribute>(FloatToStringKey);
 					UE_RETURN_ON_ERROR(FloatToStringMapMoved != nullptr, TEXT("Float->String map should be valid"));
 
-					TAttributeSetMap<FBoneTransformAnimationAttribute>* TransformToTransformMapMoved = CollectionMoved.FindSetMap<FBoneTransformAnimationAttribute>(TransformToTransformKey);
+					TBoundValueMap<FBoneTransformAnimationAttribute>* TransformToTransformMapMoved = CollectionMoved.FindSetMap<FBoneTransformAnimationAttribute>(TransformToTransformKey);
 					UE_RETURN_ON_ERROR(TransformToTransformMapMoved != nullptr, TEXT("Transform->Transform map should be valid"));
 
-					TAttributeSetMap<FFloatAnimationAttribute>* TransformToFloatMapMoved = Cast<FFloatAnimationAttribute>(CollectionMoved.FindSetMap(TransformToFloatKey));
+					TBoundValueMap<FFloatAnimationAttribute>* TransformToFloatMapMoved = Cast<FFloatAnimationAttribute>(CollectionMoved.FindSetMap(TransformToFloatKey));
 					UE_RETURN_ON_ERROR(TransformToFloatMapMoved != nullptr, TEXT("Transform->Bone map should be valid"));
 
-					TAttributeMap<FFloatAnimationAttribute>* FloatMapMoved = CollectionMoved.FindMap<FFloatAnimationAttribute>();
+					TUnboundValueMap<FFloatAnimationAttribute>* FloatMapMoved = CollectionMoved.FindMap<FFloatAnimationAttribute>();
 					UE_RETURN_ON_ERROR(FloatMapMoved != nullptr, TEXT("Float map should be valid"));
 
-					TAttributeMap<FBoneTransformAnimationAttribute>* BoneMapMoved = Cast<FBoneTransformAnimationAttribute>(CollectionMoved.FindMap(FBoneTransformAnimationAttribute::StaticStruct()));
+					TUnboundValueMap<FBoneTransformAnimationAttribute>* BoneMapMoved = Cast<FBoneTransformAnimationAttribute>(CollectionMoved.FindMap(FBoneTransformAnimationAttribute::StaticStruct()));
 					UE_RETURN_ON_ERROR(BoneMapMoved != nullptr, TEXT("Bone map should be valid"));
 
 					AddErrorIfFalse(FloatToFloatMapMoved == FloatToFloatMapCopy, TEXT("Set map instances should be moved"));
