Rename UAF ValueRuntime types

diff --git a/Engine/Plugins/Experimental/UAF/UAFTests/Source/UAFTestSuite/Private/UAF/ValueRuntime/IteratorUtilsTest.cpp b/Engine/Plugins/Experimental/UAF/UAFTests/Source/UAFTestSuite/Private/UAF/ValueRuntime/IteratorUtilsTest.cpp
--- a/Engine/Plugins/Experimental/UAF/UAFTests/Source/UAFTestSuite/Private/UAF/ValueRuntime/IteratorUtilsTest.cpp
+++ b/Engine/Plugins/Experimental/UAF/UAFTests/Source/UAFTestSuite/Private/UAF/ValueRuntime/IteratorUtilsTest.cpp
@@ -22,29 +22,29 @@
 
 namespace UE::UAF::Tests
 {
-	static bool CollectionContains(const FAttributeCollection& Collection, const FAttributeSetMap* SetMap)
+	static bool CollectionContains(const FAttributeCollection& Collection, const FBoundValueMap* Map)
 	{
-		if (SetMap == nullptr)
+		if (Map == nullptr)
 		{
 			return false;
 		}
 
-		const FAttributeSetMap* Other = Collection.FindSetMap(SetMap->GetMappingKey());
-		return SetMap == Other;
+		const FBoundValueMap* Other = Collection.FindSetMap(Map->GetMappingKey());
+		return Map == Other;
 	}
 
-	static bool CollectionContains(const FAttributeCollection& Collection, const FAttributeMap* Map)
+	static bool CollectionContains(const FAttributeCollection& Collection, const FUnboundValueMap* Map)
 	{
 		if (Map == nullptr)
 		{
 			return false;
 		}
 
-		const FAttributeMap* Other = Collection.FindMap(Map->GetValueType());
+		const FUnboundValueMap* Other = Collection.FindMap(Map->GetValueType());
 		return Map == Other;
 	}
 
-	IMPLEMENT_SIMPLE_AUTOMATION_TEST(FIteratorUtilsInnerJoinTest, "Animation.UAF.Attributes.IteratorUtils.InnerJoin", EAutomationTestFlags::EditorContext | EAutomationTestFlags::EngineFilter)
+	IMPLEMENT_SIMPLE_AUTOMATION_TEST(FIteratorUtilsInnerJoinTest, "Animation.UAF.ValueRuntime.IteratorUtils.InnerJoin", EAutomationTestFlags::EditorContext | EAutomationTestFlags::EngineFilter)
 
 	bool FIteratorUtilsInnerJoinTest::RunTest(const FString& InParameters)
 	{
@@ -119,29 +119,29 @@ namespace UE::UAF::Tests
 			AddErrorIfFalse(Collection0.GetNamedSet() == FullBodySet, TEXT("Unexpected named set"));
 
 			FAttributeMappingKey FloatToFloatKey = FAttributeMappingKey::MakeFromTo<FFloatAnimationAttribute, FFloatAnimationAttribute>();
-			TAttributeSetMap<FFloatAnimationAttribute>* FloatToFloatMap0 = Collection0.AddSetMap<FFloatAnimationAttribute>(FloatToFloatKey);
+			TBoundValueMap<FFloatAnimationAttribute>* FloatToFloatMap0 = Collection0.AddSetMap<FFloatAnimationAttribute>(FloatToFloatKey);
 			UE_RETURN_ON_ERROR(FloatToFloatMap0 != nullptr, TEXT("Float->Float map should be valid"));
 
 			FAttributeMappingKey FloatToIntegerKey = FAttributeMappingKey::MakeFromTo<FFloatAnimationAttribute, FIntegerAnimationAttribute>();
-			TAttributeSetMap<FIntegerAnimationAttribute>* FloatToIntegerMap0 = Collection0.AddSetMap<FIntegerAnimationAttribute>(FloatToIntegerKey);
+			TBoundValueMap<FIntegerAnimationAttribute>* FloatToIntegerMap0 = Collection0.AddSetMap<FIntegerAnimationAttribute>(FloatToIntegerKey);
 			UE_RETURN_ON_ERROR(FloatToIntegerMap0 != nullptr, TEXT("Float->Integer map should be valid"));
 
 			FAttributeMappingKey FloatToStringKey = FAttributeMappingKey::MakeFromTo<FFloatAnimationAttribute, FStringAnimationAttribute>();
-			TAttributeSetMap<FStringAnimationAttribute>* FloatToStringMap0 = Collection0.AddSetMap<FStringAnimationAttribute>(FloatToStringKey);
+			TBoundValueMap<FStringAnimationAttribute>* FloatToStringMap0 = Collection0.AddSetMap<FStringAnimationAttribute>(FloatToStringKey);
 			UE_RETURN_ON_ERROR(FloatToStringMap0 != nullptr, TEXT("Float->String map should be valid"));
 
 			FAttributeMappingKey TransformToTransformKey = FAttributeMappingKey::MakeFromTo<FBoneTransformAnimationAttribute, FBoneTransformAnimationAttribute>();
-			TAttributeSetMap<FBoneTransformAnimationAttribute>* TransformToTransformMap0 = Collection0.AddSetMap<FBoneTransformAnimationAttribute>(TransformToTransformKey);
+			TBoundValueMap<FBoneTransformAnimationAttribute>* TransformToTransformMap0 = Collection0.AddSetMap<FBoneTransformAnimationAttribute>(TransformToTransformKey);
 			UE_RETURN_ON_ERROR(TransformToTransformMap0 != nullptr, TEXT("Transform->Transform map should be valid"));
 
 			FAttributeMappingKey TransformToFloatKey = FAttributeMappingKey::MakeFromTo<FBoneTransformAnimationAttribute, FFloatAnimationAttribute>();
-			TAttributeSetMap<FFloatAnimationAttribute>* TransformToFloatMap0 = Collection0.AddSetMap<FFloatAnimationAttribute>(TransformToFloatKey);
+			TBoundValueMap<FFloatAnimationAttribute>* TransformToFloatMap0 = Collection0.AddSetMap<FFloatAnimationAttribute>(TransformToFloatKey);
 			UE_RETURN_ON_ERROR(TransformToFloatMap0 != nullptr, TEXT("Transform->Bone map should be valid"));
 
-			TAttributeMap<FFloatAnimationAttribute>* FloatMap0 = Collection0.AddMap<FFloatAnimationAttribute>();
+			TUnboundValueMap<FFloatAnimationAttribute>* FloatMap0 = Collection0.AddMap<FFloatAnimationAttribute>();
 			UE_RETURN_ON_ERROR(FloatMap0 != nullptr, TEXT("Float map should be valid"));
 
-			TAttributeMap<FBoneTransformAnimationAttribute>* BoneMap0 = Collection0.AddMap<FBoneTransformAnimationAttribute>();
+			TUnboundValueMap<FBoneTransformAnimationAttribute>* BoneMap0 = Collection0.AddMap<FBoneTransformAnimationAttribute>();
 			UE_RETURN_ON_ERROR(BoneMap0 != nullptr, TEXT("Bone map should be valid"));
 
 			AddErrorIfFalse(!Collection0.IsEmpty(), TEXT("Attribute collection should not be empty"));
@@ -153,25 +153,25 @@ namespace UE::UAF::Tests
 			AddErrorIfFalse(Collection1.IsEmpty(), TEXT("Attribute collection should be empty"));
 			AddErrorIfFalse(Collection1.GetNamedSet() == FullBodySet, TEXT("Unexpected named set"));
 
-			TAttributeSetMap<FFloatAnimationAttribute>* FloatToFloatMap1 = Collection1.AddSetMap<FFloatAnimationAttribute>(FloatToFloatKey);
+			TBoundValueMap<FFloatAnimationAttribute>* FloatToFloatMap1 = Collection1.AddSetMap<FFloatAnimationAttribute>(FloatToFloatKey);
 			UE_RETURN_ON_ERROR(FloatToFloatMap1 != nullptr, TEXT("Float->Float map should be valid"));
 
-			TAttributeSetMap<FIntegerAnimationAttribute>* FloatToIntegerMap1 = Collection1.AddSetMap<FIntegerAnimationAttribute>(FloatToIntegerKey);
+			TBoundValueMap<FIntegerAnimationAttribute>* FloatToIntegerMap1 = Collection1.AddSetMap<FIntegerAnimationAttribute>(FloatToIntegerKey);
 			UE_RETURN_ON_ERROR(FloatToIntegerMap1 != nullptr, TEXT("Float->Integer map should be valid"));
 
-			TAttributeSetMap<FStringAnimationAttribute>* FloatToStringMap1 = Collection1.AddSetMap<FStringAnimationAttribute>(FloatToStringKey);
+			TBoundValueMap<FStringAnimationAttribute>* FloatToStringMap1 = Collection1.AddSetMap<FStringAnimationAttribute>(FloatToStringKey);
 			UE_RETURN_ON_ERROR(FloatToStringMap1 != nullptr, TEXT("Float->String map should be valid"));
 
-			TAttributeSetMap<FBoneTransformAnimationAttribute>* TransformToTransformMap1 = Collection1.AddSetMap<FBoneTransformAnimationAttribute>(TransformToTransformKey);
+			TBoundValueMap<FBoneTransformAnimationAttribute>* TransformToTransformMap1 = Collection1.AddSetMap<FBoneTransformAnimationAttribute>(TransformToTransformKey);
 			UE_RETURN_ON_ERROR(TransformToTransformMap1 != nullptr, TEXT("Transform->Transform map should be valid"));
 
-			TAttributeSetMap<FFloatAnimationAttribute>* TransformToFloatMap1 = Collection1.AddSetMap<FFloatAnimationAttribute>(TransformToFloatKey);
+			TBoundValueMap<FFloatAnimationAttribute>* TransformToFloatMap1 = Collection1.AddSetMap<FFloatAnimationAttribute>(TransformToFloatKey);
 			UE_RETURN_ON_ERROR(TransformToFloatMap1 != nullptr, TEXT("Transform->Bone map should be valid"));
 
-			TAttributeMap<FFloatAnimationAttribute>* FloatMap1 = Collection1.AddMap<FFloatAnimationAttribute>();
+			TUnboundValueMap<FFloatAnimationAttribute>* FloatMap1 = Collection1.AddMap<FFloatAnimationAttribute>();
 			UE_RETURN_ON_ERROR(FloatMap1 != nullptr, TEXT("Float map should be valid"));
 
-			TAttributeMap<FBoneTransformAnimationAttribute>* BoneMap1 = Collection1.AddMap<FBoneTransformAnimationAttribute>();
+			TUnboundValueMap<FBoneTransformAnimationAttribute>* BoneMap1 = Collection1.AddMap<FBoneTransformAnimationAttribute>();
 			UE_RETURN_ON_ERROR(BoneMap1 != nullptr, TEXT("Bone map should be valid"));
 
 			AddErrorIfFalse(!Collection1.IsEmpty(), TEXT("Attribute collection should not be empty"));
@@ -183,7 +183,7 @@ namespace UE::UAF::Tests
 				int32 NumPredicateCalls = 0;
 				InnerJoinBy(
 					FAttributeCollectionSetMapJoinOp(),
-					[this, &NumPredicateCalls, &Collection0, &Collection1](FAttributeSetMap* Map0, FAttributeSetMap* Map1)
+					[this, &NumPredicateCalls, &Collection0, &Collection1](FBoundValueMap* Map0, FBoundValueMap* Map1)
 					{
 						NumPredicateCalls++;
 						AddErrorIfFalse(CollectionContains(Collection0, Map0), TEXT("Unexpected InnerJoin iteration result"));
@@ -198,7 +198,7 @@ namespace UE::UAF::Tests
 				NumPredicateCalls = 0;
 				InnerJoinBy(
 					FAttributeCollectionSetMapJoinOp(),
-					[this, &NumPredicateCalls, &Collection0, &Collection1](const FAttributeSetMap* Map0, const FAttributeSetMap* Map1)
+					[this, &NumPredicateCalls, &Collection0, &Collection1](const FBoundValueMap* Map0, const FBoundValueMap* Map1)
 					{
 						NumPredicateCalls++;
 						AddErrorIfFalse(CollectionContains(Collection0, Map0), TEXT("Unexpected InnerJoin iteration result"));
@@ -213,7 +213,7 @@ namespace UE::UAF::Tests
 				NumPredicateCalls = 0;
 				InnerJoinBy(
 					FAttributeCollectionSetMapJoinOp(),
-					[this, &NumPredicateCalls, &Collection0, &Collection1](const FAttributeSetMap* Map0, FAttributeSetMap* Map1)
+					[this, &NumPredicateCalls, &Collection0, &Collection1](const FBoundValueMap* Map0, FBoundValueMap* Map1)
 					{
 						NumPredicateCalls++;
 						AddErrorIfFalse(CollectionContains(Collection0, Map0), TEXT("Unexpected InnerJoin iteration result"));
@@ -228,7 +228,7 @@ namespace UE::UAF::Tests
 				NumPredicateCalls = 0;
 				InnerJoinBy(
 					FAttributeCollectionSetMapJoinOp(),
-					[this, &NumPredicateCalls, &Collection0, &Collection1](FAttributeSetMap* Map0, const FAttributeSetMap* Map1)
+					[this, &NumPredicateCalls, &Collection0, &Collection1](FBoundValueMap* Map0, const FBoundValueMap* Map1)
 					{
 						NumPredicateCalls++;
 						AddErrorIfFalse(CollectionContains(Collection0, Map0), TEXT("Unexpected InnerJoin iteration result"));
@@ -246,7 +246,7 @@ namespace UE::UAF::Tests
 				int32 NumPredicateCalls = 0;
 				InnerJoinBy(
 					FAttributeCollectionMapJoinOp(),
-					[this, &NumPredicateCalls, &Collection0, &Collection1](FAttributeMap* Map0, FAttributeMap* Map1)
+					[this, &NumPredicateCalls, &Collection0, &Collection1](FUnboundValueMap* Map0, FUnboundValueMap* Map1)
 					{
 						NumPredicateCalls++;
 						AddErrorIfFalse(CollectionContains(Collection0, Map0), TEXT("Unexpected InnerJoin iteration result"));
@@ -260,7 +260,7 @@ namespace UE::UAF::Tests
 				NumPredicateCalls = 0;
 				InnerJoinBy(
 					FAttributeCollectionMapJoinOp(),
-					[this, &NumPredicateCalls, &Collection0, &Collection1](const FAttributeMap* Map0, const FAttributeMap* Map1)
+					[this, &NumPredicateCalls, &Collection0, &Collection1](const FUnboundValueMap* Map0, const FUnboundValueMap* Map1)
 					{
 						NumPredicateCalls++;
 						AddErrorIfFalse(CollectionContains(Collection0, Map0), TEXT("Unexpected InnerJoin iteration result"));
@@ -274,7 +274,7 @@ namespace UE::UAF::Tests
 				NumPredicateCalls = 0;
 				InnerJoinBy(
 					FAttributeCollectionMapJoinOp(),
-					[this, &NumPredicateCalls, &Collection0, &Collection1](const FAttributeMap* Map0, FAttributeMap* Map1)
+					[this, &NumPredicateCalls, &Collection0, &Collection1](const FUnboundValueMap* Map0, FUnboundValueMap* Map1)
 					{
 						NumPredicateCalls++;
 						AddErrorIfFalse(CollectionContains(Collection0, Map0), TEXT("Unexpected InnerJoin iteration result"));
@@ -288,7 +288,7 @@ namespace UE::UAF::Tests
 				NumPredicateCalls = 0;
 				InnerJoinBy(
 					FAttributeCollectionMapJoinOp(),
-					[this, &NumPredicateCalls, &Collection0, &Collection1](FAttributeMap* Map0, const FAttributeMap* Map1)
+					[this, &NumPredicateCalls, &Collection0, &Collection1](FUnboundValueMap* Map0, const FUnboundValueMap* Map1)
 					{
 						NumPredicateCalls++;
 						AddErrorIfFalse(CollectionContains(Collection0, Map0), TEXT("Unexpected InnerJoin iteration result"));
@@ -305,13 +305,13 @@ namespace UE::UAF::Tests
 			AddErrorIfFalse(Collection2.IsEmpty(), TEXT("Attribute collection should be empty"));
 			AddErrorIfFalse(Collection2.GetNamedSet() == FullBodySet, TEXT("Unexpected named set"));
 
-			TAttributeSetMap<FFloatAnimationAttribute>* FloatToFloatMap2 = Collection2.AddSetMap<FFloatAnimationAttribute>(FloatToFloatKey);
+			TBoundValueMap<FFloatAnimationAttribute>* FloatToFloatMap2 = Collection2.AddSetMap<FFloatAnimationAttribute>(FloatToFloatKey);
 			UE_RETURN_ON_ERROR(FloatToFloatMap2 != nullptr, TEXT("Float->Float map should be valid"));
 
-			TAttributeSetMap<FBoneTransformAnimationAttribute>* TransformToTransformMap2 = Collection2.AddSetMap<FBoneTransformAnimationAttribute>(TransformToTransformKey);
+			TBoundValueMap<FBoneTransformAnimationAttribute>* TransformToTransformMap2 = Collection2.AddSetMap<FBoneTransformAnimationAttribute>(TransformToTransformKey);
 			UE_RETURN_ON_ERROR(TransformToTransformMap2 != nullptr, TEXT("Transform->Transform map should be valid"));
 
-			TAttributeMap<FFloatAnimationAttribute>* FloatMap2 = Collection2.AddMap<FFloatAnimationAttribute>();
+			TUnboundValueMap<FFloatAnimationAttribute>* FloatMap2 = Collection2.AddMap<FFloatAnimationAttribute>();
 			UE_RETURN_ON_ERROR(FloatMap2 != nullptr, TEXT("Float map should be valid"));
 
 			AddErrorIfFalse(!Collection2.IsEmpty(), TEXT("Attribute collection should not be empty"));
@@ -323,7 +323,7 @@ namespace UE::UAF::Tests
 				int32 NumPredicateCalls = 0;
 				InnerJoinBy(
 					FAttributeCollectionSetMapJoinOp(),
-					[this, &NumPredicateCalls, &Collection0, &Collection2](FAttributeSetMap* Map2, FAttributeSetMap* Map0)
+					[this, &NumPredicateCalls, &Collection0, &Collection2](FBoundValueMap* Map2, FBoundValueMap* Map0)
 					{
 						NumPredicateCalls++;
 						AddErrorIfFalse(CollectionContains(Collection0, Map0), TEXT("Unexpected InnerJoin iteration result"));
@@ -338,7 +338,7 @@ namespace UE::UAF::Tests
 				NumPredicateCalls = 0;
 				InnerJoinBy(
 					FAttributeCollectionSetMapJoinOp(),
-					[this, &NumPredicateCalls, &Collection0, &Collection2](const FAttributeSetMap* Map2, const FAttributeSetMap* Map0)
+					[this, &NumPredicateCalls, &Collection0, &Collection2](const FBoundValueMap* Map2, const FBoundValueMap* Map0)
 					{
 						NumPredicateCalls++;
 						AddErrorIfFalse(CollectionContains(Collection0, Map0), TEXT("Unexpected InnerJoin iteration result"));
@@ -353,7 +353,7 @@ namespace UE::UAF::Tests
 				NumPredicateCalls = 0;
 				InnerJoinBy(
 					FAttributeCollectionSetMapJoinOp(),
-					[this, &NumPredicateCalls, &Collection0, &Collection2](const FAttributeSetMap* Map2, FAttributeSetMap* Map0)
+					[this, &NumPredicateCalls, &Collection0, &Collection2](const FBoundValueMap* Map2, FBoundValueMap* Map0)
 					{
 						NumPredicateCalls++;
 						AddErrorIfFalse(CollectionContains(Collection0, Map0), TEXT("Unexpected InnerJoin iteration result"));
@@ -368,7 +368,7 @@ namespace UE::UAF::Tests
 				NumPredicateCalls = 0;
 				InnerJoinBy(
 					FAttributeCollectionSetMapJoinOp(),
-					[this, &NumPredicateCalls, &Collection0, &Collection2](FAttributeSetMap* Map2, const FAttributeSetMap* Map0)
+					[this, &NumPredicateCalls, &Collection0, &Collection2](FBoundValueMap* Map2, const FBoundValueMap* Map0)
 					{
 						NumPredicateCalls++;
 						AddErrorIfFalse(CollectionContains(Collection0, Map0), TEXT("Unexpected InnerJoin iteration result"));
@@ -386,7 +386,7 @@ namespace UE::UAF::Tests
 				int32 NumPredicateCalls = 0;
 				InnerJoinBy(
 					FAttributeCollectionMapJoinOp(),
-					[this, &NumPredicateCalls, &Collection0, &Collection2](FAttributeMap* Map2, FAttributeMap* Map0)
+					[this, &NumPredicateCalls, &Collection0, &Collection2](FUnboundValueMap* Map2, FUnboundValueMap* Map0)
 					{
 						NumPredicateCalls++;
 						AddErrorIfFalse(CollectionContains(Collection0, Map0), TEXT("Unexpected InnerJoin iteration result"));
@@ -400,7 +400,7 @@ namespace UE::UAF::Tests
 				NumPredicateCalls = 0;
 				InnerJoinBy(
 					FAttributeCollectionMapJoinOp(),
-					[this, &NumPredicateCalls, &Collection0, &Collection2](const FAttributeMap* Map2, const FAttributeMap* Map0)
+					[this, &NumPredicateCalls, &Collection0, &Collection2](const FUnboundValueMap* Map2, const FUnboundValueMap* Map0)
 					{
 						NumPredicateCalls++;
 						AddErrorIfFalse(CollectionContains(Collection0, Map0), TEXT("Unexpected InnerJoin iteration result"));
@@ -414,7 +414,7 @@ namespace UE::UAF::Tests
 				NumPredicateCalls = 0;
 				InnerJoinBy(
 					FAttributeCollectionMapJoinOp(),
-					[this, &NumPredicateCalls, &Collection0, &Collection2](const FAttributeMap* Map2, FAttributeMap* Map0)
+					[this, &NumPredicateCalls, &Collection0, &Collection2](const FUnboundValueMap* Map2, FUnboundValueMap* Map0)
 					{
 						NumPredicateCalls++;
 						AddErrorIfFalse(CollectionContains(Collection0, Map0), TEXT("Unexpected InnerJoin iteration result"));
@@ -428,7 +428,7 @@ namespace UE::UAF::Tests
 				NumPredicateCalls = 0;
 				InnerJoinBy(
 					FAttributeCollectionMapJoinOp(),
-					[this, &NumPredicateCalls, &Collection0, &Collection2](FAttributeMap* Map2, const FAttributeMap* Map0)
+					[this, &NumPredicateCalls, &Collection0, &Collection2](FUnboundValueMap* Map2, const FUnboundValueMap* Map0)
 					{
 						NumPredicateCalls++;
 						AddErrorIfFalse(CollectionContains(Collection0, Map0), TEXT("Unexpected InnerJoin iteration result"));
@@ -445,7 +445,7 @@ namespace UE::UAF::Tests
 				int32 NumPredicateCalls = 0;
 				InnerJoinBy(
 					FAttributeCollectionSetMapJoinOp(),
-					[this, &NumPredicateCalls, &Collection0, &Collection2](FAttributeSetMap* Map0, FAttributeSetMap* Map2)
+					[this, &NumPredicateCalls, &Collection0, &Collection2](FBoundValueMap* Map0, FBoundValueMap* Map2)
 					{
 						NumPredicateCalls++;
 						AddErrorIfFalse(CollectionContains(Collection0, Map0), TEXT("Unexpected InnerJoin iteration result"));
@@ -460,7 +460,7 @@ namespace UE::UAF::Tests
 				NumPredicateCalls = 0;
 				InnerJoinBy(
 					FAttributeCollectionSetMapJoinOp(),
-					[this, &NumPredicateCalls, &Collection0, &Collection2](const FAttributeSetMap* Map0, const FAttributeSetMap* Map2)
+					[this, &NumPredicateCalls, &Collection0, &Collection2](const FBoundValueMap* Map0, const FBoundValueMap* Map2)
 					{
 						NumPredicateCalls++;
 						AddErrorIfFalse(CollectionContains(Collection0, Map0), TEXT("Unexpected InnerJoin iteration result"));
@@ -475,7 +475,7 @@ namespace UE::UAF::Tests
 				NumPredicateCalls = 0;
 				InnerJoinBy(
 					FAttributeCollectionSetMapJoinOp(),
-					[this, &NumPredicateCalls, &Collection0, &Collection2](const FAttributeSetMap* Map0, FAttributeSetMap* Map2)
+					[this, &NumPredicateCalls, &Collection0, &Collection2](const FBoundValueMap* Map0, FBoundValueMap* Map2)
 					{
 						NumPredicateCalls++;
 						AddErrorIfFalse(CollectionContains(Collection0, Map0), TEXT("Unexpected InnerJoin iteration result"));
@@ -490,7 +490,7 @@ namespace UE::UAF::Tests
 				NumPredicateCalls = 0;
 				InnerJoinBy(
 					FAttributeCollectionSetMapJoinOp(),
-					[this, &NumPredicateCalls, &Collection0, &Collection2](FAttributeSetMap* Map0, const FAttributeSetMap* Map2)
+					[this, &NumPredicateCalls, &Collection0, &Collection2](FBoundValueMap* Map0, const FBoundValueMap* Map2)
 					{
 						NumPredicateCalls++;
 						AddErrorIfFalse(CollectionContains(Collection0, Map0), TEXT("Unexpected InnerJoin iteration result"));
@@ -508,7 +508,7 @@ namespace UE::UAF::Tests
 				int32 NumPredicateCalls = 0;
 				InnerJoinBy(
 					FAttributeCollectionMapJoinOp(),
-					[this, &NumPredicateCalls, &Collection0, &Collection2](FAttributeMap* Map0, FAttributeMap* Map2)
+					[this, &NumPredicateCalls, &Collection0, &Collection2](FUnboundValueMap* Map0, FUnboundValueMap* Map2)
 					{
 						NumPredicateCalls++;
 						AddErrorIfFalse(CollectionContains(Collection0, Map0), TEXT("Unexpected InnerJoin iteration result"));
@@ -522,7 +522,7 @@ namespace UE::UAF::Tests
 				NumPredicateCalls = 0;
 				InnerJoinBy(
 					FAttributeCollectionMapJoinOp(),
-					[this, &NumPredicateCalls, &Collection0, &Collection2](const FAttributeMap* Map0, const FAttributeMap* Map2)
+					[this, &NumPredicateCalls, &Collection0, &Collection2](const FUnboundValueMap* Map0, const FUnboundValueMap* Map2)
 					{
 						NumPredicateCalls++;
 						AddErrorIfFalse(CollectionContains(Collection0, Map0), TEXT("Unexpected InnerJoin iteration result"));
@@ -536,7 +536,7 @@ namespace UE::UAF::Tests
 				NumPredicateCalls = 0;
 				InnerJoinBy(
 					FAttributeCollectionMapJoinOp(),
-					[this, &NumPredicateCalls, &Collection0, &Collection2](const FAttributeMap* Map0, FAttributeMap* Map2)
+					[this, &NumPredicateCalls, &Collection0, &Collection2](const FUnboundValueMap* Map0, FUnboundValueMap* Map2)
 					{
 						NumPredicateCalls++;
 						AddErrorIfFalse(CollectionContains(Collection0, Map0), TEXT("Unexpected InnerJoin iteration result"));
@@ -550,7 +550,7 @@ namespace UE::UAF::Tests
 				NumPredicateCalls = 0;
 				InnerJoinBy(
 					FAttributeCollectionMapJoinOp(),
-					[this, &NumPredicateCalls, &Collection0, &Collection2](FAttributeMap* Map0, const FAttributeMap* Map2)
+					[this, &NumPredicateCalls, &Collection0, &Collection2](FUnboundValueMap* Map0, const FUnboundValueMap* Map2)
 					{
 						NumPredicateCalls++;
 						AddErrorIfFalse(CollectionContains(Collection0, Map0), TEXT("Unexpected InnerJoin iteration result"));
@@ -567,7 +567,7 @@ namespace UE::UAF::Tests
 				int32 NumPredicateCalls = 0;
 				InnerJoinBy(
 					FAttributeCollectionSetMapJoinOp(),
-					[this, &NumPredicateCalls, &Collection0, &Collection1, &Collection2](FAttributeSetMap* Map0, FAttributeSetMap* Map1, FAttributeSetMap* Map2)
+					[this, &NumPredicateCalls, &Collection0, &Collection1, &Collection2](FBoundValueMap* Map0, FBoundValueMap* Map1, FBoundValueMap* Map2)
 					{
 						NumPredicateCalls++;
 						AddErrorIfFalse(CollectionContains(Collection0, Map0), TEXT("Unexpected InnerJoin iteration result"));
@@ -586,7 +586,7 @@ namespace UE::UAF::Tests
 				NumPredicateCalls = 0;
 				InnerJoinBy(
 					FAttributeCollectionSetMapJoinOp(),
-					[this, &NumPredicateCalls, &Collection0, &Collection1, &Collection2](FAttributeSetMap* Map0, FAttributeSetMap* Map2, FAttributeSetMap* Map1)
+					[this, &NumPredicateCalls, &Collection0, &Collection1, &Collection2](FBoundValueMap* Map0, FBoundValueMap* Map2, FBoundValueMap* Map1)
 					{
 						NumPredicateCalls++;
 						AddErrorIfFalse(CollectionContains(Collection0, Map0), TEXT("Unexpected InnerJoin iteration result"));
@@ -605,7 +605,7 @@ namespace UE::UAF::Tests
 				NumPredicateCalls = 0;
 				InnerJoinBy(
 					FAttributeCollectionSetMapJoinOp(),
-					[this, &NumPredicateCalls, &Collection0, &Collection1, &Collection2](FAttributeSetMap* Map1, FAttributeSetMap* Map0, FAttributeSetMap* Map2)
+					[this, &NumPredicateCalls, &Collection0, &Collection1, &Collection2](FBoundValueMap* Map1, FBoundValueMap* Map0, FBoundValueMap* Map2)
 					{
 						NumPredicateCalls++;
 						AddErrorIfFalse(CollectionContains(Collection0, Map0), TEXT("Unexpected InnerJoin iteration result"));
@@ -624,7 +624,7 @@ namespace UE::UAF::Tests
 				NumPredicateCalls = 0;
 				InnerJoinBy(
 					FAttributeCollectionSetMapJoinOp(),
-					[this, &NumPredicateCalls, &Collection0, &Collection1, &Collection2](FAttributeSetMap* Map1, FAttributeSetMap* Map2, FAttributeSetMap* Map0)
+					[this, &NumPredicateCalls, &Collection0, &Collection1, &Collection2](FBoundValueMap* Map1, FBoundValueMap* Map2, FBoundValueMap* Map0)
 					{
 						NumPredicateCalls++;
 						AddErrorIfFalse(CollectionContains(Collection0, Map0), TEXT("Unexpected InnerJoin iteration result"));
@@ -643,7 +643,7 @@ namespace UE::UAF::Tests
 				NumPredicateCalls = 0;
 				InnerJoinBy(
 					FAttributeCollectionSetMapJoinOp(),
-					[this, &NumPredicateCalls, &Collection0, &Collection1, &Collection2](FAttributeSetMap* Map2, FAttributeSetMap* Map0, FAttributeSetMap* Map1)
+					[this, &NumPredicateCalls, &Collection0, &Collection1, &Collection2](FBoundValueMap* Map2, FBoundValueMap* Map0, FBoundValueMap* Map1)
 					{
 						NumPredicateCalls++;
 						AddErrorIfFalse(CollectionContains(Collection0, Map0), TEXT("Unexpected InnerJoin iteration result"));
@@ -662,7 +662,7 @@ namespace UE::UAF::Tests
 				NumPredicateCalls = 0;
 				InnerJoinBy(
 					FAttributeCollectionSetMapJoinOp(),
-					[this, &NumPredicateCalls, &Collection0, &Collection1, &Collection2](FAttributeSetMap* Map2, FAttributeSetMap* Map1, FAttributeSetMap* Map0)
+					[this, &NumPredicateCalls, &Collection0, &Collection1, &Collection2](FBoundValueMap* Map2, FBoundValueMap* Map1, FBoundValueMap* Map0)
 					{
 						NumPredicateCalls++;
 						AddErrorIfFalse(CollectionContains(Collection0, Map0), TEXT("Unexpected InnerJoin iteration result"));
@@ -684,7 +684,7 @@ namespace UE::UAF::Tests
 				int32 NumPredicateCalls = 0;
 				InnerJoinBy(
 					FAttributeCollectionMapJoinOp(),
-					[this, &NumPredicateCalls, &Collection0, &Collection1, &Collection2](FAttributeMap* Map0, FAttributeMap* Map1, FAttributeMap* Map2)
+					[this, &NumPredicateCalls, &Collection0, &Collection1, &Collection2](FUnboundValueMap* Map0, FUnboundValueMap* Map1, FUnboundValueMap* Map2)
 					{
 						NumPredicateCalls++;
 						AddErrorIfFalse(CollectionContains(Collection0, Map0), TEXT("Unexpected InnerJoin iteration result"));
@@ -701,7 +701,7 @@ namespace UE::UAF::Tests
 				NumPredicateCalls = 0;
 				InnerJoinBy(
 					FAttributeCollectionMapJoinOp(),
-					[this, &NumPredicateCalls, &Collection0, &Collection1, &Collection2](FAttributeMap* Map0, FAttributeMap* Map2, FAttributeMap* Map1)
+					[this, &NumPredicateCalls, &Collection0, &Collection1, &Collection2](FUnboundValueMap* Map0, FUnboundValueMap* Map2, FUnboundValueMap* Map1)
 					{
 						NumPredicateCalls++;
 						AddErrorIfFalse(CollectionContains(Collection0, Map0), TEXT("Unexpected InnerJoin iteration result"));
@@ -718,7 +718,7 @@ namespace UE::UAF::Tests
 				NumPredicateCalls = 0;
 				InnerJoinBy(
 					FAttributeCollectionMapJoinOp(),
-					[this, &NumPredicateCalls, &Collection0, &Collection1, &Collection2](FAttributeMap* Map1, FAttributeMap* Map0, FAttributeMap* Map2)
+					[this, &NumPredicateCalls, &Collection0, &Collection1, &Collection2](FUnboundValueMap* Map1, FUnboundValueMap* Map0, FUnboundValueMap* Map2)
 					{
 						NumPredicateCalls++;
 						AddErrorIfFalse(CollectionContains(Collection0, Map0), TEXT("Unexpected InnerJoin iteration result"));
@@ -735,7 +735,7 @@ namespace UE::UAF::Tests
 				NumPredicateCalls = 0;
 				InnerJoinBy(
 					FAttributeCollectionMapJoinOp(),
-					[this, &NumPredicateCalls, &Collection0, &Collection1, &Collection2](FAttributeMap* Map1, FAttributeMap* Map2, FAttributeMap* Map0)
+					[this, &NumPredicateCalls, &Collection0, &Collection1, &Collection2](FUnboundValueMap* Map1, FUnboundValueMap* Map2, FUnboundValueMap* Map0)
 					{
 						NumPredicateCalls++;
 						AddErrorIfFalse(CollectionContains(Collection0, Map0), TEXT("Unexpected InnerJoin iteration result"));
@@ -752,7 +752,7 @@ namespace UE::UAF::Tests
 				NumPredicateCalls = 0;
 				InnerJoinBy(
 					FAttributeCollectionMapJoinOp(),
-					[this, &NumPredicateCalls, &Collection0, &Collection1, &Collection2](FAttributeMap* Map2, FAttributeMap* Map0, FAttributeMap* Map1)
+					[this, &NumPredicateCalls, &Collection0, &Collection1, &Collection2](FUnboundValueMap* Map2, FUnboundValueMap* Map0, FUnboundValueMap* Map1)
 					{
 						NumPredicateCalls++;
 						AddErrorIfFalse(CollectionContains(Collection0, Map0), TEXT("Unexpected InnerJoin iteration result"));
@@ -769,7 +769,7 @@ namespace UE::UAF::Tests
 				NumPredicateCalls = 0;
 				InnerJoinBy(
 					FAttributeCollectionMapJoinOp(),
-					[this, &NumPredicateCalls, &Collection0, &Collection1, &Collection2](FAttributeMap* Map2, FAttributeMap* Map1, FAttributeMap* Map0)
+					[this, &NumPredicateCalls, &Collection0, &Collection1, &Collection2](FUnboundValueMap* Map2, FUnboundValueMap* Map1, FUnboundValueMap* Map0)
 					{
 						NumPredicateCalls++;
 						AddErrorIfFalse(CollectionContains(Collection0, Map0), TEXT("Unexpected InnerJoin iteration result"));
@@ -796,7 +796,7 @@ namespace UE::UAF::Tests
 					using TransformerType = Transformers::FSanitize;
 					using ValueType = FFloatAnimationAttribute;
 
-					static void TransformSetMap(FAttributeSetMap* SetMap) {}	// No-op
+					static void TransformBoundValueMap(FBoundValueMap* SetMap) {}	// No-op
 				};
 
 				// Dummy vector sanitize transformer
@@ -805,7 +805,7 @@ namespace UE::UAF::Tests
 					using TransformerType = Transformers::FSanitize;
 					using ValueType = FVectorAnimationAttribute;
 
-					static void TransformSetMap(FAttributeSetMap* SetMap) {}	// No-op
+					static void TransformBoundValueMap(FBoundValueMap* SetMap) {}	// No-op
 				};
 
 #ifdef _MSC_VER
@@ -813,9 +813,9 @@ namespace UE::UAF::Tests
 	#pragma warning(disable : 4191) // unsafe conversion from 'function pointer type A' to 'function pointer type B'
 #endif
 
-				FAttributeValueTransformerList TransformerList(Transformers::FSanitize::TransformerName);
-				TransformerList.AddSetMapTransformer(FFloatAnimationAttribute::StaticStruct(), reinterpret_cast<FRawTransformerFunc>(&FFloatAttributeTransformer_SanitizeTest::TransformSetMap));
-				TransformerList.AddSetMapTransformer(FVectorAnimationAttribute::StaticStruct(), reinterpret_cast<FRawTransformerFunc>(&FVectorAttributeTransformer_SanitizeTest::TransformSetMap));
+				FValueTransformerList TransformerList(Transformers::FSanitize::TransformerName);
+				TransformerList.AddBoundValueMapTransformer(FFloatAnimationAttribute::StaticStruct(), reinterpret_cast<FRawTransformerFunc>(&FFloatAttributeTransformer_SanitizeTest::TransformBoundValueMap));
+				TransformerList.AddBoundValueMapTransformer(FVectorAnimationAttribute::StaticStruct(), reinterpret_cast<FRawTransformerFunc>(&FVectorAttributeTransformer_SanitizeTest::TransformBoundValueMap));
 
 #ifdef _MSC_VER
 	#pragma warning(pop)
@@ -823,14 +823,14 @@ namespace UE::UAF::Tests
 
 				int32 NumPredicateCalls = 0;
 				InnerJoinBy(
-					FAttributeValueTransformerListWithAttributeCollectionSetMapJoinOp(),
-					[this, &NumPredicateCalls, &Collection0](FRawTransformerFunc Transformer, FAttributeSetMap* SetMap)
+					FValueTransformerListWithBoundMapCollectionJoinOp(),
+					[this, &NumPredicateCalls, &Collection0](FRawTransformerFunc Transformer, FBoundValueMap* SetMap)
 					{
 						NumPredicateCalls++;
 						AddErrorIfFalse(CollectionContains(Collection0, SetMap), TEXT("Unexpected InnerJoin iteration result"));
 						AddErrorIfFalse(SetMap->GetValueType() == FFloatAnimationAttribute::StaticStruct(), TEXT("Unexpected InnerJoin iteration result"));
 					},
-					TransformerList.CreateSetMapTransformerIterator(),
+					TransformerList.CreateBoundValueMapTransformerIterator(),
 					Collection0.CreateSetMapIterator());
 				AddErrorIfFalse(NumPredicateCalls == 2, TEXT("Unexpected InnerJoin iteration result"));
 			}
@@ -847,7 +847,7 @@ namespace UE::UAF::Tests
 					using TransformerType = Transformers::FSanitize;
 					using ValueType = FFloatAnimationAttribute;
 
-					static void TransformMap(FAttributeMap* Map) {}	// No-op
+					static void TransformMap(FUnboundValueMap* Map) {}	// No-op
 				};
 
 				// Dummy vector sanitize transformer
@@ -856,7 +856,7 @@ namespace UE::UAF::Tests
 					using TransformerType = Transformers::FSanitize;
 					using ValueType = FVectorAnimationAttribute;
 
-					static void TransformMap(FAttributeMap* Map) {}	// No-op
+					static void TransformMap(FUnboundValueMap* Map) {}	// No-op
 				};
 
 #ifdef _MSC_VER
@@ -864,9 +864,9 @@ namespace UE::UAF::Tests
 	#pragma warning(disable : 4191) // unsafe conversion from 'function pointer type A' to 'function pointer type B'
 #endif
 
-				FAttributeValueTransformerList TransformerList(Transformers::FSanitize::TransformerName);
-				TransformerList.AddMapTransformer(FFloatAnimationAttribute::StaticStruct(), reinterpret_cast<FRawTransformerFunc>(&FFloatAttributeTransformer_SanitizeTest::TransformMap));
-				TransformerList.AddMapTransformer(FVectorAnimationAttribute::StaticStruct(), reinterpret_cast<FRawTransformerFunc>(&FVectorAttributeTransformer_SanitizeTest::TransformMap));
+				FValueTransformerList TransformerList(Transformers::FSanitize::TransformerName);
+				TransformerList.AddUnboundValueMapTransformer(FFloatAnimationAttribute::StaticStruct(), reinterpret_cast<FRawTransformerFunc>(&FFloatAttributeTransformer_SanitizeTest::TransformMap));
+				TransformerList.AddUnboundValueMapTransformer(FVectorAnimationAttribute::StaticStruct(), reinterpret_cast<FRawTransformerFunc>(&FVectorAttributeTransformer_SanitizeTest::TransformMap));
 
 #ifdef _MSC_VER
 	#pragma warning(pop)
@@ -874,14 +874,14 @@ namespace UE::UAF::Tests
 
 				int32 NumPredicateCalls = 0;
 				InnerJoinBy(
-					FAttributeValueTransformerListWithAttributeCollectionMapJoinOp(),
-					[this, &NumPredicateCalls, &Collection0](FRawTransformerFunc Transformer, FAttributeMap* Map)
+					FValueTransformerListWithUnboundMapCollectionJoinOp(),
+					[this, &NumPredicateCalls, &Collection0](FRawTransformerFunc Transformer, FUnboundValueMap* Map)
 					{
 						NumPredicateCalls++;
 						AddErrorIfFalse(CollectionContains(Collection0, Map), TEXT("Unexpected InnerJoin iteration result"));
 						AddErrorIfFalse(Map->GetValueType() == FFloatAnimationAttribute::StaticStruct(), TEXT("Unexpected InnerJoin iteration result"));
 					},
-					TransformerList.CreateMapTransformerIterator(),
+					TransformerList.CreateUnboundValueMapTransformerIterator(),
 					Collection0.CreateMapIterator());
 				AddErrorIfFalse(NumPredicateCalls == 1, TEXT("Unexpected InnerJoin iteration result"));
 			}
@@ -890,7 +890,7 @@ namespace UE::UAF::Tests
 		return true;
 	}
 
-	IMPLEMENT_SIMPLE_AUTOMATION_TEST(FIteratorUtilsOuterJoinTest, "Animation.UAF.Attributes.IteratorUtils.OuterJoin", EAutomationTestFlags::EditorContext | EAutomationTestFlags::EngineFilter)
+	IMPLEMENT_SIMPLE_AUTOMATION_TEST(FIteratorUtilsOuterJoinTest, "Animation.UAF.ValueRuntime.IteratorUtils.OuterJoin", EAutomationTestFlags::EditorContext | EAutomationTestFlags::EngineFilter)
 
 	bool FIteratorUtilsOuterJoinTest::RunTest(const FString& InParameters)
 	{
@@ -965,29 +965,29 @@ namespace UE::UAF::Tests
 			AddErrorIfFalse(Collection0.GetNamedSet() == FullBodySet, TEXT("Unexpected named set"));
 
 			FAttributeMappingKey FloatToFloatKey = FAttributeMappingKey::MakeFromTo<FFloatAnimationAttribute, FFloatAnimationAttribute>();
-			TAttributeSetMap<FFloatAnimationAttribute>* FloatToFloatMap0 = Collection0.AddSetMap<FFloatAnimationAttribute>(FloatToFloatKey);
+			TBoundValueMap<FFloatAnimationAttribute>* FloatToFloatMap0 = Collection0.AddSetMap<FFloatAnimationAttribute>(FloatToFloatKey);
 			UE_RETURN_ON_ERROR(FloatToFloatMap0 != nullptr, TEXT("Float->Float map should be valid"));
 
 			FAttributeMappingKey FloatToIntegerKey = FAttributeMappingKey::MakeFromTo<FFloatAnimationAttribute, FIntegerAnimationAttribute>();
-			TAttributeSetMap<FIntegerAnimationAttribute>* FloatToIntegerMap0 = Collection0.AddSetMap<FIntegerAnimationAttribute>(FloatToIntegerKey);
+			TBoundValueMap<FIntegerAnimationAttribute>* FloatToIntegerMap0 = Collection0.AddSetMap<FIntegerAnimationAttribute>(FloatToIntegerKey);
 			UE_RETURN_ON_ERROR(FloatToIntegerMap0 != nullptr, TEXT("Float->Integer map should be valid"));
 
 			FAttributeMappingKey FloatToStringKey = FAttributeMappingKey::MakeFromTo<FFloatAnimationAttribute, FStringAnimationAttribute>();
-			TAttributeSetMap<FStringAnimationAttribute>* FloatToStringMap0 = Collection0.AddSetMap<FStringAnimationAttribute>(FloatToStringKey);
+			TBoundValueMap<FStringAnimationAttribute>* FloatToStringMap0 = Collection0.AddSetMap<FStringAnimationAttribute>(FloatToStringKey);
 			UE_RETURN_ON_ERROR(FloatToStringMap0 != nullptr, TEXT("Float->String map should be valid"));
 
 			FAttributeMappingKey TransformToTransformKey = FAttributeMappingKey::MakeFromTo<FBoneTransformAnimationAttribute, FBoneTransformAnimationAttribute>();
-			TAttributeSetMap<FBoneTransformAnimationAttribute>* TransformToTransformMap0 = Collection0.AddSetMap<FBoneTransformAnimationAttribute>(TransformToTransformKey);
+			TBoundValueMap<FBoneTransformAnimationAttribute>* TransformToTransformMap0 = Collection0.AddSetMap<FBoneTransformAnimationAttribute>(TransformToTransformKey);
 			UE_RETURN_ON_ERROR(TransformToTransformMap0 != nullptr, TEXT("Transform->Transform map should be valid"));
 
 			FAttributeMappingKey TransformToFloatKey = FAttributeMappingKey::MakeFromTo<FBoneTransformAnimationAttribute, FFloatAnimationAttribute>();
-			TAttributeSetMap<FFloatAnimationAttribute>* TransformToFloatMap0 = Collection0.AddSetMap<FFloatAnimationAttribute>(TransformToFloatKey);
+			TBoundValueMap<FFloatAnimationAttribute>* TransformToFloatMap0 = Collection0.AddSetMap<FFloatAnimationAttribute>(TransformToFloatKey);
 			UE_RETURN_ON_ERROR(TransformToFloatMap0 != nullptr, TEXT("Transform->Bone map should be valid"));
 
-			TAttributeMap<FFloatAnimationAttribute>* FloatMap0 = Collection0.AddMap<FFloatAnimationAttribute>();
+			TUnboundValueMap<FFloatAnimationAttribute>* FloatMap0 = Collection0.AddMap<FFloatAnimationAttribute>();
 			UE_RETURN_ON_ERROR(FloatMap0 != nullptr, TEXT("Float map should be valid"));
 
-			TAttributeMap<FBoneTransformAnimationAttribute>* BoneMap0 = Collection0.AddMap<FBoneTransformAnimationAttribute>();
+			TUnboundValueMap<FBoneTransformAnimationAttribute>* BoneMap0 = Collection0.AddMap<FBoneTransformAnimationAttribute>();
 			UE_RETURN_ON_ERROR(BoneMap0 != nullptr, TEXT("Bone map should be valid"));
 
 			AddErrorIfFalse(!Collection0.IsEmpty(), TEXT("Attribute collection should not be empty"));
@@ -999,25 +999,25 @@ namespace UE::UAF::Tests
 			AddErrorIfFalse(Collection1.IsEmpty(), TEXT("Attribute collection should be empty"));
 			AddErrorIfFalse(Collection1.GetNamedSet() == FullBodySet, TEXT("Unexpected named set"));
 
-			TAttributeSetMap<FFloatAnimationAttribute>* FloatToFloatMap1 = Collection1.AddSetMap<FFloatAnimationAttribute>(FloatToFloatKey);
+			TBoundValueMap<FFloatAnimationAttribute>* FloatToFloatMap1 = Collection1.AddSetMap<FFloatAnimationAttribute>(FloatToFloatKey);
 			UE_RETURN_ON_ERROR(FloatToFloatMap1 != nullptr, TEXT("Float->Float map should be valid"));
 
-			TAttributeSetMap<FIntegerAnimationAttribute>* FloatToIntegerMap1 = Collection1.AddSetMap<FIntegerAnimationAttribute>(FloatToIntegerKey);
+			TBoundValueMap<FIntegerAnimationAttribute>* FloatToIntegerMap1 = Collection1.AddSetMap<FIntegerAnimationAttribute>(FloatToIntegerKey);
 			UE_RETURN_ON_ERROR(FloatToIntegerMap1 != nullptr, TEXT("Float->Integer map should be valid"));
 
-			TAttributeSetMap<FStringAnimationAttribute>* FloatToStringMap1 = Collection1.AddSetMap<FStringAnimationAttribute>(FloatToStringKey);
+			TBoundValueMap<FStringAnimationAttribute>* FloatToStringMap1 = Collection1.AddSetMap<FStringAnimationAttribute>(FloatToStringKey);
 			UE_RETURN_ON_ERROR(FloatToStringMap1 != nullptr, TEXT("Float->String map should be valid"));
 
-			TAttributeSetMap<FBoneTransformAnimationAttribute>* TransformToTransformMap1 = Collection1.AddSetMap<FBoneTransformAnimationAttribute>(TransformToTransformKey);
+			TBoundValueMap<FBoneTransformAnimationAttribute>* TransformToTransformMap1 = Collection1.AddSetMap<FBoneTransformAnimationAttribute>(TransformToTransformKey);
 			UE_RETURN_ON_ERROR(TransformToTransformMap1 != nullptr, TEXT("Transform->Transform map should be valid"));
 
-			TAttributeSetMap<FFloatAnimationAttribute>* TransformToFloatMap1 = Collection1.AddSetMap<FFloatAnimationAttribute>(TransformToFloatKey);
+			TBoundValueMap<FFloatAnimationAttribute>* TransformToFloatMap1 = Collection1.AddSetMap<FFloatAnimationAttribute>(TransformToFloatKey);
 			UE_RETURN_ON_ERROR(TransformToFloatMap1 != nullptr, TEXT("Transform->Bone map should be valid"));
 
-			TAttributeMap<FFloatAnimationAttribute>* FloatMap1 = Collection1.AddMap<FFloatAnimationAttribute>();
+			TUnboundValueMap<FFloatAnimationAttribute>* FloatMap1 = Collection1.AddMap<FFloatAnimationAttribute>();
 			UE_RETURN_ON_ERROR(FloatMap1 != nullptr, TEXT("Float map should be valid"));
 
-			TAttributeMap<FBoneTransformAnimationAttribute>* BoneMap1 = Collection1.AddMap<FBoneTransformAnimationAttribute>();
+			TUnboundValueMap<FBoneTransformAnimationAttribute>* BoneMap1 = Collection1.AddMap<FBoneTransformAnimationAttribute>();
 			UE_RETURN_ON_ERROR(BoneMap1 != nullptr, TEXT("Bone map should be valid"));
 
 			AddErrorIfFalse(!Collection1.IsEmpty(), TEXT("Attribute collection should not be empty"));
@@ -1029,7 +1029,7 @@ namespace UE::UAF::Tests
 				int32 NumPredicateCalls = 0;
 				OuterJoinBy(
 					FAttributeCollectionSetMapJoinOp(),
-					[this, &NumPredicateCalls, &Collection0, &Collection1](FAttributeSetMap* Map0, FAttributeSetMap* Map1)
+					[this, &NumPredicateCalls, &Collection0, &Collection1](FBoundValueMap* Map0, FBoundValueMap* Map1)
 					{
 						NumPredicateCalls++;
 						AddErrorIfFalse(CollectionContains(Collection0, Map0), TEXT("Unexpected OuterJoin iteration result"));
@@ -1044,7 +1044,7 @@ namespace UE::UAF::Tests
 				NumPredicateCalls = 0;
 				OuterJoinBy(
 					FAttributeCollectionSetMapJoinOp(),
-					[this, &NumPredicateCalls, &Collection0, &Collection1](const FAttributeSetMap* Map0, const FAttributeSetMap* Map1)
+					[this, &NumPredicateCalls, &Collection0, &Collection1](const FBoundValueMap* Map0, const FBoundValueMap* Map1)
 					{
 						NumPredicateCalls++;
 						AddErrorIfFalse(CollectionContains(Collection0, Map0), TEXT("Unexpected OuterJoin iteration result"));
@@ -1059,7 +1059,7 @@ namespace UE::UAF::Tests
 				NumPredicateCalls = 0;
 				OuterJoinBy(
 					FAttributeCollectionSetMapJoinOp(),
-					[this, &NumPredicateCalls, &Collection0, &Collection1](const FAttributeSetMap* Map0, FAttributeSetMap* Map1)
+					[this, &NumPredicateCalls, &Collection0, &Collection1](const FBoundValueMap* Map0, FBoundValueMap* Map1)
 					{
 						NumPredicateCalls++;
 						AddErrorIfFalse(CollectionContains(Collection0, Map0), TEXT("Unexpected OuterJoin iteration result"));
@@ -1074,7 +1074,7 @@ namespace UE::UAF::Tests
 				NumPredicateCalls = 0;
 				OuterJoinBy(
 					FAttributeCollectionSetMapJoinOp(),
-					[this, &NumPredicateCalls, &Collection0, &Collection1](FAttributeSetMap* Map0, const FAttributeSetMap* Map1)
+					[this, &NumPredicateCalls, &Collection0, &Collection1](FBoundValueMap* Map0, const FBoundValueMap* Map1)
 					{
 						NumPredicateCalls++;
 						AddErrorIfFalse(CollectionContains(Collection0, Map0), TEXT("Unexpected OuterJoin iteration result"));
@@ -1092,7 +1092,7 @@ namespace UE::UAF::Tests
 				int32 NumPredicateCalls = 0;
 				OuterJoinBy(
 					FAttributeCollectionMapJoinOp(),
-					[this, &NumPredicateCalls, &Collection0, &Collection1](FAttributeMap* Map0, FAttributeMap* Map1)
+					[this, &NumPredicateCalls, &Collection0, &Collection1](FUnboundValueMap* Map0, FUnboundValueMap* Map1)
 					{
 						NumPredicateCalls++;
 						AddErrorIfFalse(CollectionContains(Collection0, Map0), TEXT("Unexpected OuterJoin iteration result"));
@@ -1106,7 +1106,7 @@ namespace UE::UAF::Tests
 				NumPredicateCalls = 0;
 				OuterJoinBy(
 					FAttributeCollectionMapJoinOp(),
-					[this, &NumPredicateCalls, &Collection0, &Collection1](const FAttributeMap* Map0, const FAttributeMap* Map1)
+					[this, &NumPredicateCalls, &Collection0, &Collection1](const FUnboundValueMap* Map0, const FUnboundValueMap* Map1)
 					{
 						NumPredicateCalls++;
 						AddErrorIfFalse(CollectionContains(Collection0, Map0), TEXT("Unexpected OuterJoin iteration result"));
@@ -1120,7 +1120,7 @@ namespace UE::UAF::Tests
 				NumPredicateCalls = 0;
 				OuterJoinBy(
 					FAttributeCollectionMapJoinOp(),
-					[this, &NumPredicateCalls, &Collection0, &Collection1](const FAttributeMap* Map0, FAttributeMap* Map1)
+					[this, &NumPredicateCalls, &Collection0, &Collection1](const FUnboundValueMap* Map0, FUnboundValueMap* Map1)
 					{
 						NumPredicateCalls++;
 						AddErrorIfFalse(CollectionContains(Collection0, Map0), TEXT("Unexpected OuterJoin iteration result"));
@@ -1134,7 +1134,7 @@ namespace UE::UAF::Tests
 				NumPredicateCalls = 0;
 				OuterJoinBy(
 					FAttributeCollectionMapJoinOp(),
-					[this, &NumPredicateCalls, &Collection0, &Collection1](FAttributeMap* Map0, const FAttributeMap* Map1)
+					[this, &NumPredicateCalls, &Collection0, &Collection1](FUnboundValueMap* Map0, const FUnboundValueMap* Map1)
 					{
 						NumPredicateCalls++;
 						AddErrorIfFalse(CollectionContains(Collection0, Map0), TEXT("Unexpected OuterJoin iteration result"));
@@ -1151,13 +1151,13 @@ namespace UE::UAF::Tests
 			AddErrorIfFalse(Collection2.IsEmpty(), TEXT("Attribute collection should be empty"));
 			AddErrorIfFalse(Collection2.GetNamedSet() == FullBodySet, TEXT("Unexpected named set"));
 
-			TAttributeSetMap<FFloatAnimationAttribute>* FloatToFloatMap2 = Collection2.AddSetMap<FFloatAnimationAttribute>(FloatToFloatKey);
+			TBoundValueMap<FFloatAnimationAttribute>* FloatToFloatMap2 = Collection2.AddSetMap<FFloatAnimationAttribute>(FloatToFloatKey);
 			UE_RETURN_ON_ERROR(FloatToFloatMap2 != nullptr, TEXT("Float->Float map should be valid"));
 
-			TAttributeSetMap<FBoneTransformAnimationAttribute>* TransformToTransformMap2 = Collection2.AddSetMap<FBoneTransformAnimationAttribute>(TransformToTransformKey);
+			TBoundValueMap<FBoneTransformAnimationAttribute>* TransformToTransformMap2 = Collection2.AddSetMap<FBoneTransformAnimationAttribute>(TransformToTransformKey);
 			UE_RETURN_ON_ERROR(TransformToTransformMap2 != nullptr, TEXT("Transform->Transform map should be valid"));
 
-			TAttributeMap<FFloatAnimationAttribute>* FloatMap2 = Collection2.AddMap<FFloatAnimationAttribute>();
+			TUnboundValueMap<FFloatAnimationAttribute>* FloatMap2 = Collection2.AddMap<FFloatAnimationAttribute>();
 			UE_RETURN_ON_ERROR(FloatMap2 != nullptr, TEXT("Float map should be valid"));
 
 			AddErrorIfFalse(!Collection2.IsEmpty(), TEXT("Attribute collection should not be empty"));
@@ -1170,7 +1170,7 @@ namespace UE::UAF::Tests
 				int32 NumNullEntries = 0;
 				OuterJoinBy(
 					FAttributeCollectionSetMapJoinOp(),
-					[this, &NumPredicateCalls, &NumNullEntries, &Collection0, &Collection2](FAttributeSetMap* Map2, FAttributeSetMap* Map0)
+					[this, &NumPredicateCalls, &NumNullEntries, &Collection0, &Collection2](FBoundValueMap* Map2, FBoundValueMap* Map0)
 					{
 						NumPredicateCalls++;
 						AddErrorIfFalse(CollectionContains(Collection0, Map0), TEXT("Unexpected OuterJoin iteration result"));
@@ -1196,7 +1196,7 @@ namespace UE::UAF::Tests
 				NumNullEntries = 0;
 				OuterJoinBy(
 					FAttributeCollectionSetMapJoinOp(),
-					[this, &NumPredicateCalls, &NumNullEntries, &Collection0, &Collection2](const FAttributeSetMap* Map2, const FAttributeSetMap* Map0)
+					[this, &NumPredicateCalls, &NumNullEntries, &Collection0, &Collection2](const FBoundValueMap* Map2, const FBoundValueMap* Map0)
 					{
 						NumPredicateCalls++;
 						AddErrorIfFalse(CollectionContains(Collection0, Map0), TEXT("Unexpected OuterJoin iteration result"));
@@ -1222,7 +1222,7 @@ namespace UE::UAF::Tests
 				NumNullEntries = 0;
 				OuterJoinBy(
 					FAttributeCollectionSetMapJoinOp(),
-					[this, &NumPredicateCalls, &NumNullEntries, &Collection0, &Collection2](const FAttributeSetMap* Map2, FAttributeSetMap* Map0)
+					[this, &NumPredicateCalls, &NumNullEntries, &Collection0, &Collection2](const FBoundValueMap* Map2, FBoundValueMap* Map0)
 					{
 						NumPredicateCalls++;
 						AddErrorIfFalse(CollectionContains(Collection0, Map0), TEXT("Unexpected OuterJoin iteration result"));
@@ -1248,7 +1248,7 @@ namespace UE::UAF::Tests
 				NumNullEntries = 0;
 				OuterJoinBy(
 					FAttributeCollectionSetMapJoinOp(),
-					[this, &NumPredicateCalls, &NumNullEntries, &Collection0, &Collection2](FAttributeSetMap* Map2, const FAttributeSetMap* Map0)
+					[this, &NumPredicateCalls, &NumNullEntries, &Collection0, &Collection2](FBoundValueMap* Map2, const FBoundValueMap* Map0)
 					{
 						NumPredicateCalls++;
 						AddErrorIfFalse(CollectionContains(Collection0, Map0), TEXT("Unexpected OuterJoin iteration result"));
@@ -1277,7 +1277,7 @@ namespace UE::UAF::Tests
 				int32 NumNullEntries = 0;
 				OuterJoinBy(
 					FAttributeCollectionMapJoinOp(),
-					[this, &NumPredicateCalls, &NumNullEntries, &Collection0, &Collection2](FAttributeMap* Map2, FAttributeMap* Map0)
+					[this, &NumPredicateCalls, &NumNullEntries, &Collection0, &Collection2](FUnboundValueMap* Map2, FUnboundValueMap* Map0)
 					{
 						NumPredicateCalls++;
 						AddErrorIfFalse(CollectionContains(Collection0, Map0), TEXT("Unexpected OuterJoin iteration result"));
@@ -1302,7 +1302,7 @@ namespace UE::UAF::Tests
 				NumNullEntries = 0;
 				OuterJoinBy(
 					FAttributeCollectionMapJoinOp(),
-					[this, &NumPredicateCalls, &NumNullEntries, &Collection0, &Collection2](const FAttributeMap* Map2, const FAttributeMap* Map0)
+					[this, &NumPredicateCalls, &NumNullEntries, &Collection0, &Collection2](const FUnboundValueMap* Map2, const FUnboundValueMap* Map0)
 					{
 						NumPredicateCalls++;
 						AddErrorIfFalse(CollectionContains(Collection0, Map0), TEXT("Unexpected OuterJoin iteration result"));
@@ -1327,7 +1327,7 @@ namespace UE::UAF::Tests
 				NumNullEntries = 0;
 				OuterJoinBy(
 					FAttributeCollectionMapJoinOp(),
-					[this, &NumPredicateCalls, &NumNullEntries, &Collection0, &Collection2](const FAttributeMap* Map2, FAttributeMap* Map0)
+					[this, &NumPredicateCalls, &NumNullEntries, &Collection0, &Collection2](const FUnboundValueMap* Map2, FUnboundValueMap* Map0)
 					{
 						NumPredicateCalls++;
 						AddErrorIfFalse(CollectionContains(Collection0, Map0), TEXT("Unexpected OuterJoin iteration result"));
@@ -1352,7 +1352,7 @@ namespace UE::UAF::Tests
 				NumNullEntries = 0;
 				OuterJoinBy(
 					FAttributeCollectionMapJoinOp(),
-					[this, &NumPredicateCalls, &NumNullEntries, &Collection0, &Collection2](FAttributeMap* Map2, const FAttributeMap* Map0)
+					[this, &NumPredicateCalls, &NumNullEntries, &Collection0, &Collection2](FUnboundValueMap* Map2, const FUnboundValueMap* Map0)
 					{
 						NumPredicateCalls++;
 						AddErrorIfFalse(CollectionContains(Collection0, Map0), TEXT("Unexpected OuterJoin iteration result"));
@@ -1380,7 +1380,7 @@ namespace UE::UAF::Tests
 				int32 NumNullEntries = 0;
 				OuterJoinBy(
 					FAttributeCollectionSetMapJoinOp(),
-					[this, &NumPredicateCalls, &NumNullEntries, &Collection0, &Collection2](FAttributeSetMap* Map0, FAttributeSetMap* Map2)
+					[this, &NumPredicateCalls, &NumNullEntries, &Collection0, &Collection2](FBoundValueMap* Map0, FBoundValueMap* Map2)
 					{
 						NumPredicateCalls++;
 						AddErrorIfFalse(CollectionContains(Collection0, Map0), TEXT("Unexpected OuterJoin iteration result"));
@@ -1406,7 +1406,7 @@ namespace UE::UAF::Tests
 				NumNullEntries = 0;
 				OuterJoinBy(
 					FAttributeCollectionSetMapJoinOp(),
-					[this, &NumPredicateCalls, &NumNullEntries, &Collection0, &Collection2](const FAttributeSetMap* Map0, const FAttributeSetMap* Map2)
+					[this, &NumPredicateCalls, &NumNullEntries, &Collection0, &Collection2](const FBoundValueMap* Map0, const FBoundValueMap* Map2)
 					{
 						NumPredicateCalls++;
 						AddErrorIfFalse(CollectionContains(Collection0, Map0), TEXT("Unexpected OuterJoin iteration result"));
@@ -1432,7 +1432,7 @@ namespace UE::UAF::Tests
 				NumNullEntries = 0;
 				OuterJoinBy(
 					FAttributeCollectionSetMapJoinOp(),
-					[this, &NumPredicateCalls, &NumNullEntries, &Collection0, &Collection2](const FAttributeSetMap* Map0, FAttributeSetMap* Map2)
+					[this, &NumPredicateCalls, &NumNullEntries, &Collection0, &Collection2](const FBoundValueMap* Map0, FBoundValueMap* Map2)
 					{
 						NumPredicateCalls++;
 						AddErrorIfFalse(CollectionContains(Collection0, Map0), TEXT("Unexpected OuterJoin iteration result"));
@@ -1458,7 +1458,7 @@ namespace UE::UAF::Tests
 				NumNullEntries = 0;
 				OuterJoinBy(
 					FAttributeCollectionSetMapJoinOp(),
-					[this, &NumPredicateCalls, &NumNullEntries, &Collection0, &Collection2](FAttributeSetMap* Map0, const FAttributeSetMap* Map2)
+					[this, &NumPredicateCalls, &NumNullEntries, &Collection0, &Collection2](FBoundValueMap* Map0, const FBoundValueMap* Map2)
 					{
 						NumPredicateCalls++;
 						AddErrorIfFalse(CollectionContains(Collection0, Map0), TEXT("Unexpected OuterJoin iteration result"));
@@ -1487,7 +1487,7 @@ namespace UE::UAF::Tests
 				int32 NumNullEntries = 0;
 				OuterJoinBy(
 					FAttributeCollectionMapJoinOp(),
-					[this, &NumPredicateCalls, &NumNullEntries, &Collection0, &Collection2](FAttributeMap* Map0, FAttributeMap* Map2)
+					[this, &NumPredicateCalls, &NumNullEntries, &Collection0, &Collection2](FUnboundValueMap* Map0, FUnboundValueMap* Map2)
 					{
 						NumPredicateCalls++;
 						AddErrorIfFalse(CollectionContains(Collection0, Map0), TEXT("Unexpected OuterJoin iteration result"));
@@ -1512,7 +1512,7 @@ namespace UE::UAF::Tests
 				NumNullEntries = 0;
 				OuterJoinBy(
 					FAttributeCollectionMapJoinOp(),
-					[this, &NumPredicateCalls, &NumNullEntries, &Collection0, &Collection2](const FAttributeMap* Map0, const FAttributeMap* Map2)
+					[this, &NumPredicateCalls, &NumNullEntries, &Collection0, &Collection2](const FUnboundValueMap* Map0, const FUnboundValueMap* Map2)
 					{
 						NumPredicateCalls++;
 						AddErrorIfFalse(CollectionContains(Collection0, Map0), TEXT("Unexpected OuterJoin iteration result"));
@@ -1537,7 +1537,7 @@ namespace UE::UAF::Tests
 				NumNullEntries = 0;
 				OuterJoinBy(
 					FAttributeCollectionMapJoinOp(),
-					[this, &NumPredicateCalls, &NumNullEntries, &Collection0, &Collection2](const FAttributeMap* Map0, FAttributeMap* Map2)
+					[this, &NumPredicateCalls, &NumNullEntries, &Collection0, &Collection2](const FUnboundValueMap* Map0, FUnboundValueMap* Map2)
 					{
 						NumPredicateCalls++;
 						AddErrorIfFalse(CollectionContains(Collection0, Map0), TEXT("Unexpected OuterJoin iteration result"));
@@ -1562,7 +1562,7 @@ namespace UE::UAF::Tests
 				NumNullEntries = 0;
 				OuterJoinBy(
 					FAttributeCollectionMapJoinOp(),
-					[this, &NumPredicateCalls, &NumNullEntries, &Collection0, &Collection2](FAttributeMap* Map0, const FAttributeMap* Map2)
+					[this, &NumPredicateCalls, &NumNullEntries, &Collection0, &Collection2](FUnboundValueMap* Map0, const FUnboundValueMap* Map2)
 					{
 						NumPredicateCalls++;
 						AddErrorIfFalse(CollectionContains(Collection0, Map0), TEXT("Unexpected OuterJoin iteration result"));
@@ -1596,7 +1596,7 @@ namespace UE::UAF::Tests
 					using TransformerType = Transformers::FSanitize;
 					using ValueType = FFloatAnimationAttribute;
 
-					static void TransformSetMap(FAttributeSetMap* SetMap) {}	// No-op
+					static void TransformBoundValueMap(FBoundValueMap* SetMap) {}	// No-op
 				};
 
 				// Dummy vector sanitize transformer
@@ -1605,7 +1605,7 @@ namespace UE::UAF::Tests
 					using TransformerType = Transformers::FSanitize;
 					using ValueType = FVectorAnimationAttribute;
 
-					static void TransformSetMap(FAttributeSetMap* SetMap) {}	// No-op
+					static void TransformBoundValueMap(FBoundValueMap* SetMap) {}	// No-op
 				};
 
 #ifdef _MSC_VER
@@ -1613,9 +1613,9 @@ namespace UE::UAF::Tests
 	#pragma warning(disable : 4191) // unsafe conversion from 'function pointer type A' to 'function pointer type B'
 #endif
 
-				FAttributeValueTransformerList TransformerList(Transformers::FSanitize::TransformerName);
-				TransformerList.AddSetMapTransformer(FFloatAnimationAttribute::StaticStruct(), reinterpret_cast<FRawTransformerFunc>(&FFloatAttributeTransformer_SanitizeTest::TransformSetMap));
-				TransformerList.AddSetMapTransformer(FVectorAnimationAttribute::StaticStruct(), reinterpret_cast<FRawTransformerFunc>(&FVectorAttributeTransformer_SanitizeTest::TransformSetMap));
+				FValueTransformerList TransformerList(Transformers::FSanitize::TransformerName);
+				TransformerList.AddBoundValueMapTransformer(FFloatAnimationAttribute::StaticStruct(), reinterpret_cast<FRawTransformerFunc>(&FFloatAttributeTransformer_SanitizeTest::TransformBoundValueMap));
+				TransformerList.AddBoundValueMapTransformer(FVectorAnimationAttribute::StaticStruct(), reinterpret_cast<FRawTransformerFunc>(&FVectorAttributeTransformer_SanitizeTest::TransformBoundValueMap));
 
 #ifdef _MSC_VER
 	#pragma warning(pop)
@@ -1625,8 +1625,8 @@ namespace UE::UAF::Tests
 				int32 NumNullCollectionEntries = 0;
 				int32 NumNullTransformEntries = 0;
 				OuterJoinBy(
-					FAttributeValueTransformerListWithAttributeCollectionSetMapJoinOp(),
-					[this, &NumPredicateCalls, &NumNullCollectionEntries, &NumNullTransformEntries, &Collection0](FRawTransformerFunc Transformer, FAttributeSetMap* SetMap)
+					FValueTransformerListWithBoundMapCollectionJoinOp(),
+					[this, &NumPredicateCalls, &NumNullCollectionEntries, &NumNullTransformEntries, &Collection0](FRawTransformerFunc Transformer, FBoundValueMap* SetMap)
 					{
 						NumPredicateCalls++;
 
@@ -1649,7 +1649,7 @@ namespace UE::UAF::Tests
 							NumNullTransformEntries++;
 						}
 					},
-					TransformerList.CreateSetMapTransformerIterator(),
+					TransformerList.CreateBoundValueMapTransformerIterator(),
 					Collection0.CreateSetMapIterator());
 
 				// [Transformer Value Type, Collection Set Map Value Type]
@@ -1679,7 +1679,7 @@ namespace UE::UAF::Tests
 					using TransformerType = Transformers::FSanitize;
 					using ValueType = FFloatAnimationAttribute;
 
-					static void TransformMap(FAttributeMap* Map) {}	// No-op
+					static void TransformMap(FUnboundValueMap* Map) {}	// No-op
 				};
 
 				// Dummy vector sanitize transformer
@@ -1688,7 +1688,7 @@ namespace UE::UAF::Tests
 					using TransformerType = Transformers::FSanitize;
 					using ValueType = FVectorAnimationAttribute;
 
-					static void TransformMap(FAttributeMap* Map) {}	// No-op
+					static void TransformMap(FUnboundValueMap* Map) {}	// No-op
 				};
 
 #ifdef _MSC_VER
@@ -1696,9 +1696,9 @@ namespace UE::UAF::Tests
 	#pragma warning(disable : 4191) // unsafe conversion from 'function pointer type A' to 'function pointer type B'
 #endif
 
-				FAttributeValueTransformerList TransformerList(Transformers::FSanitize::TransformerName);
-				TransformerList.AddMapTransformer(FFloatAnimationAttribute::StaticStruct(), reinterpret_cast<FRawTransformerFunc>(&FFloatAttributeTransformer_SanitizeTest::TransformMap));
-				TransformerList.AddMapTransformer(FVectorAnimationAttribute::StaticStruct(), reinterpret_cast<FRawTransformerFunc>(&FVectorAttributeTransformer_SanitizeTest::TransformMap));
+				FValueTransformerList TransformerList(Transformers::FSanitize::TransformerName);
+				TransformerList.AddUnboundValueMapTransformer(FFloatAnimationAttribute::StaticStruct(), reinterpret_cast<FRawTransformerFunc>(&FFloatAttributeTransformer_SanitizeTest::TransformMap));
+				TransformerList.AddUnboundValueMapTransformer(FVectorAnimationAttribute::StaticStruct(), reinterpret_cast<FRawTransformerFunc>(&FVectorAttributeTransformer_SanitizeTest::TransformMap));
 
 #ifdef _MSC_VER
 	#pragma warning(pop)
@@ -1708,8 +1708,8 @@ namespace UE::UAF::Tests
 				int32 NumNullCollectionEntries = 0;
 				int32 NumNullTransformEntries = 0;
 				OuterJoinBy(
-					FAttributeValueTransformerListWithAttributeCollectionMapJoinOp(),
-					[this, &NumPredicateCalls, &NumNullCollectionEntries, &NumNullTransformEntries, &Collection0](FRawTransformerFunc Transformer, FAttributeMap* Map)
+					FValueTransformerListWithUnboundMapCollectionJoinOp(),
+					[this, &NumPredicateCalls, &NumNullCollectionEntries, &NumNullTransformEntries, &Collection0](FRawTransformerFunc Transformer, FUnboundValueMap* Map)
 					{
 						NumPredicateCalls++;
 
@@ -1732,7 +1732,7 @@ namespace UE::UAF::Tests
 							NumNullTransformEntries++;
 						}
 					},
-					TransformerList.CreateMapTransformerIterator(),
+					TransformerList.CreateUnboundValueMapTransformerIterator(),
 					Collection0.CreateMapIterator());
 
 				// [Transformer Value Type, Collection Map Value Type]
