Rename UAF attribute collection and split into 3 types

diff --git a/Engine/Plugins/Experimental/UAF/UAFTests/Source/UAFTestSuite/Private/UAF/ValueRuntime/IteratorUtilsTest.cpp b/Engine/Plugins/Experimental/UAF/UAFTests/Source/UAFTestSuite/Private/UAF/ValueRuntime/IteratorUtilsTest.cpp
--- a/Engine/Plugins/Experimental/UAF/UAFTests/Source/UAFTestSuite/Private/UAF/ValueRuntime/IteratorUtilsTest.cpp
+++ b/Engine/Plugins/Experimental/UAF/UAFTests/Source/UAFTestSuite/Private/UAF/ValueRuntime/IteratorUtilsTest.cpp
@@ -13,7 +13,7 @@
 #include "UAF/AbstractSkeleton/AbstractSkeletonSetCollection.h"
 #include "UAF/Attributes/AttributeBindingData.h"
 #include "UAF/Attributes/EngineAttributes.h"
-#include "UAF/ValueRuntime/AttributeCollection.h"
+#include "UAF/ValueRuntime/ValueBundle.h"
 #include "UAF/ValueRuntime/IteratorUtils.h"
 #include "UAF/ValueRuntime/Transformers/Transformer.h"
 #include "UAF/ValueRuntime/Transformers/Sanitize.h"
@@ -22,25 +22,25 @@
 
 namespace UE::UAF::Tests
 {
-	static bool CollectionContains(const FAttributeCollection& Collection, const FBoundValueMap* Map)
+	static bool CollectionContains(const FBoundMapCollection& Collection, const FBoundValueMap* Map)
 	{
 		if (Map == nullptr)
 		{
 			return false;
 		}
 
-		const FBoundValueMap* Other = Collection.FindSetMap(Map->GetMappingKey());
+		const FBoundValueMap* Other = Collection.Find(Map->GetMappingKey());
 		return Map == Other;
 	}
 
-	static bool CollectionContains(const FAttributeCollection& Collection, const FUnboundValueMap* Map)
+	static bool CollectionContains(const FUnboundMapCollection& Collection, const FUnboundValueMap* Map)
 	{
 		if (Map == nullptr)
 		{
 			return false;
 		}
 
-		const FUnboundValueMap* Other = Collection.FindMap(Map->GetValueType());
+		const FUnboundValueMap* Other = Collection.Find(Map->GetValueType());
 		return Map == Other;
 	}
 
@@ -109,80 +109,67 @@
 			UE_RETURN_ON_ERROR(BindingData.IsValid(), TEXT("Attribute binding data should be valid"));
 		}
 
+		// Bound value map collection
 		{
 			FAttributeNamedSetPtr FullBodySet = BindingData->FindNamedSet(TEXT("FullBody"));
 			UE_RETURN_ON_ERROR(FullBodySet.IsValid(), TEXT("FullBody named set should be valid"));
 
-			FAttributeCollectionHeap Collection0(FullBodySet);
+			FBoundMapCollectionHeap Collection0(FullBodySet);
 
-			AddErrorIfFalse(Collection0.IsEmpty(), TEXT("Attribute collection should be empty"));
+			AddErrorIfFalse(Collection0.IsEmpty(), TEXT("Collection should be empty"));
 			AddErrorIfFalse(Collection0.GetNamedSet() == FullBodySet, TEXT("Unexpected named set"));
 
 			FAttributeMappingKey FloatToFloatKey = FAttributeMappingKey::MakeFromTo<FFloatAnimationAttribute, FFloatAnimationAttribute>();
-			TBoundValueMap<FFloatAnimationAttribute>* FloatToFloatMap0 = Collection0.AddSetMap<FFloatAnimationAttribute>(FloatToFloatKey);
+			TBoundValueMap<FFloatAnimationAttribute>* FloatToFloatMap0 = Collection0.Add<FFloatAnimationAttribute>(FloatToFloatKey);
 			UE_RETURN_ON_ERROR(FloatToFloatMap0 != nullptr, TEXT("Float->Float map should be valid"));
 
 			FAttributeMappingKey FloatToIntegerKey = FAttributeMappingKey::MakeFromTo<FFloatAnimationAttribute, FIntegerAnimationAttribute>();
-			TBoundValueMap<FIntegerAnimationAttribute>* FloatToIntegerMap0 = Collection0.AddSetMap<FIntegerAnimationAttribute>(FloatToIntegerKey);
+			TBoundValueMap<FIntegerAnimationAttribute>* FloatToIntegerMap0 = Collection0.Add<FIntegerAnimationAttribute>(FloatToIntegerKey);
 			UE_RETURN_ON_ERROR(FloatToIntegerMap0 != nullptr, TEXT("Float->Integer map should be valid"));
 
 			FAttributeMappingKey FloatToStringKey = FAttributeMappingKey::MakeFromTo<FFloatAnimationAttribute, FStringAnimationAttribute>();
-			TBoundValueMap<FStringAnimationAttribute>* FloatToStringMap0 = Collection0.AddSetMap<FStringAnimationAttribute>(FloatToStringKey);
+			TBoundValueMap<FStringAnimationAttribute>* FloatToStringMap0 = Collection0.Add<FStringAnimationAttribute>(FloatToStringKey);
 			UE_RETURN_ON_ERROR(FloatToStringMap0 != nullptr, TEXT("Float->String map should be valid"));
 
 			FAttributeMappingKey TransformToTransformKey = FAttributeMappingKey::MakeFromTo<FBoneTransformAnimationAttribute, FBoneTransformAnimationAttribute>();
-			TBoundValueMap<FBoneTransformAnimationAttribute>* TransformToTransformMap0 = Collection0.AddSetMap<FBoneTransformAnimationAttribute>(TransformToTransformKey);
+			TBoundValueMap<FBoneTransformAnimationAttribute>* TransformToTransformMap0 = Collection0.Add<FBoneTransformAnimationAttribute>(TransformToTransformKey);
 			UE_RETURN_ON_ERROR(TransformToTransformMap0 != nullptr, TEXT("Transform->Transform map should be valid"));
 
 			FAttributeMappingKey TransformToFloatKey = FAttributeMappingKey::MakeFromTo<FBoneTransformAnimationAttribute, FFloatAnimationAttribute>();
-			TBoundValueMap<FFloatAnimationAttribute>* TransformToFloatMap0 = Collection0.AddSetMap<FFloatAnimationAttribute>(TransformToFloatKey);
+			TBoundValueMap<FFloatAnimationAttribute>* TransformToFloatMap0 = Collection0.Add<FFloatAnimationAttribute>(TransformToFloatKey);
 			UE_RETURN_ON_ERROR(TransformToFloatMap0 != nullptr, TEXT("Transform->Bone map should be valid"));
 
-			TUnboundValueMap<FFloatAnimationAttribute>* FloatMap0 = Collection0.AddMap<FFloatAnimationAttribute>();
-			UE_RETURN_ON_ERROR(FloatMap0 != nullptr, TEXT("Float map should be valid"));
+			AddErrorIfFalse(!Collection0.IsEmpty(), TEXT("Collection should not be empty"));
+			AddErrorIfFalse(Collection0.Num() == 5, TEXT("Unexpected map count"));
 
-			TUnboundValueMap<FBoneTransformAnimationAttribute>* BoneMap0 = Collection0.AddMap<FBoneTransformAnimationAttribute>();
-			UE_RETURN_ON_ERROR(BoneMap0 != nullptr, TEXT("Bone map should be valid"));
+			FBoundMapCollectionHeap Collection1(FullBodySet);
 
-			AddErrorIfFalse(!Collection0.IsEmpty(), TEXT("Attribute collection should not be empty"));
-			AddErrorIfFalse(Collection0.NumSetMaps() == 5, TEXT("Unexpected set map count"));
-			AddErrorIfFalse(Collection0.NumMaps() == 2, TEXT("Unexpected set count"));
-
-			FAttributeCollectionHeap Collection1(FullBodySet);
-
-			AddErrorIfFalse(Collection1.IsEmpty(), TEXT("Attribute collection should be empty"));
+			AddErrorIfFalse(Collection1.IsEmpty(), TEXT("Collection should be empty"));
 			AddErrorIfFalse(Collection1.GetNamedSet() == FullBodySet, TEXT("Unexpected named set"));
 
-			TBoundValueMap<FFloatAnimationAttribute>* FloatToFloatMap1 = Collection1.AddSetMap<FFloatAnimationAttribute>(FloatToFloatKey);
+			TBoundValueMap<FFloatAnimationAttribute>* FloatToFloatMap1 = Collection1.Add<FFloatAnimationAttribute>(FloatToFloatKey);
 			UE_RETURN_ON_ERROR(FloatToFloatMap1 != nullptr, TEXT("Float->Float map should be valid"));
 
-			TBoundValueMap<FIntegerAnimationAttribute>* FloatToIntegerMap1 = Collection1.AddSetMap<FIntegerAnimationAttribute>(FloatToIntegerKey);
+			TBoundValueMap<FIntegerAnimationAttribute>* FloatToIntegerMap1 = Collection1.Add<FIntegerAnimationAttribute>(FloatToIntegerKey);
 			UE_RETURN_ON_ERROR(FloatToIntegerMap1 != nullptr, TEXT("Float->Integer map should be valid"));
 
-			TBoundValueMap<FStringAnimationAttribute>* FloatToStringMap1 = Collection1.AddSetMap<FStringAnimationAttribute>(FloatToStringKey);
+			TBoundValueMap<FStringAnimationAttribute>* FloatToStringMap1 = Collection1.Add<FStringAnimationAttribute>(FloatToStringKey);
 			UE_RETURN_ON_ERROR(FloatToStringMap1 != nullptr, TEXT("Float->String map should be valid"));
 
-			TBoundValueMap<FBoneTransformAnimationAttribute>* TransformToTransformMap1 = Collection1.AddSetMap<FBoneTransformAnimationAttribute>(TransformToTransformKey);
+			TBoundValueMap<FBoneTransformAnimationAttribute>* TransformToTransformMap1 = Collection1.Add<FBoneTransformAnimationAttribute>(TransformToTransformKey);
 			UE_RETURN_ON_ERROR(TransformToTransformMap1 != nullptr, TEXT("Transform->Transform map should be valid"));
 
-			TBoundValueMap<FFloatAnimationAttribute>* TransformToFloatMap1 = Collection1.AddSetMap<FFloatAnimationAttribute>(TransformToFloatKey);
+			TBoundValueMap<FFloatAnimationAttribute>* TransformToFloatMap1 = Collection1.Add<FFloatAnimationAttribute>(TransformToFloatKey);
 			UE_RETURN_ON_ERROR(TransformToFloatMap1 != nullptr, TEXT("Transform->Bone map should be valid"));
 
-			TUnboundValueMap<FFloatAnimationAttribute>* FloatMap1 = Collection1.AddMap<FFloatAnimationAttribute>();
-			UE_RETURN_ON_ERROR(FloatMap1 != nullptr, TEXT("Float map should be valid"));
+			AddErrorIfFalse(!Collection1.IsEmpty(), TEXT("Collection should not be empty"));
+			AddErrorIfFalse(Collection1.Num() == 5, TEXT("Unexpected map count"));
 
-			TUnboundValueMap<FBoneTransformAnimationAttribute>* BoneMap1 = Collection1.AddMap<FBoneTransformAnimationAttribute>();
-			UE_RETURN_ON_ERROR(BoneMap1 != nullptr, TEXT("Bone map should be valid"));
-
-			AddErrorIfFalse(!Collection1.IsEmpty(), TEXT("Attribute collection should not be empty"));
-			AddErrorIfFalse(Collection1.NumSetMaps() == 5, TEXT("Unexpected set map count"));
-			AddErrorIfFalse(Collection1.NumMaps() == 2, TEXT("Unexpected set count"));
-
-			// Both inputs perfectly match (set maps)
+			// Both inputs perfectly match
 			{
 				int32 NumPredicateCalls = 0;
 				InnerJoinBy(
-					FAttributeCollectionSetMapJoinOp(),
+					FBoundMapCollectionJoinOp(),
 					[this, &NumPredicateCalls, &Collection0, &Collection1](FBoundValueMap* Map0, FBoundValueMap* Map1)
 					{
 						NumPredicateCalls++;
@@ -191,13 +178,13 @@
 						AddErrorIfFalse(Map0 != nullptr && Map1 != nullptr && Map0->GetAttributeType() == Map1->GetAttributeType(), TEXT("Unexpected InnerJoin iteration result"));
 						AddErrorIfFalse(Map0 != nullptr && Map1 != nullptr && Map0->GetValueType() == Map1->GetValueType(), TEXT("Unexpected InnerJoin iteration result"));
 					},
-					Collection0.CreateSetMapIterator(),
-					Collection1.CreateSetMapIterator());
+					Collection0.CreateIterator(),
+					Collection1.CreateIterator());
 				AddErrorIfFalse(NumPredicateCalls == 5, TEXT("Unexpected InnerJoin iteration result"));
 
 				NumPredicateCalls = 0;
 				InnerJoinBy(
-					FAttributeCollectionSetMapJoinOp(),
+					FBoundMapCollectionJoinOp(),
 					[this, &NumPredicateCalls, &Collection0, &Collection1](const FBoundValueMap* Map0, const FBoundValueMap* Map1)
 					{
 						NumPredicateCalls++;
@@ -206,13 +193,13 @@
 						AddErrorIfFalse(Map0 != nullptr && Map1 != nullptr && Map0->GetAttributeType() == Map1->GetAttributeType(), TEXT("Unexpected InnerJoin iteration result"));
 						AddErrorIfFalse(Map0 != nullptr && Map1 != nullptr && Map0->GetValueType() == Map1->GetValueType(), TEXT("Unexpected InnerJoin iteration result"));
 					},
-					Collection0.CreateConstSetMapIterator(),
-					Collection1.CreateConstSetMapIterator());
+					Collection0.CreateConstIterator(),
+					Collection1.CreateConstIterator());
 				AddErrorIfFalse(NumPredicateCalls == 5, TEXT("Unexpected InnerJoin iteration result"));
 
 				NumPredicateCalls = 0;
 				InnerJoinBy(
-					FAttributeCollectionSetMapJoinOp(),
+					FBoundMapCollectionJoinOp(),
 					[this, &NumPredicateCalls, &Collection0, &Collection1](const FBoundValueMap* Map0, FBoundValueMap* Map1)
 					{
 						NumPredicateCalls++;
@@ -221,13 +208,13 @@
 						AddErrorIfFalse(Map0 != nullptr && Map1 != nullptr && Map0->GetAttributeType() == Map1->GetAttributeType(), TEXT("Unexpected InnerJoin iteration result"));
 						AddErrorIfFalse(Map0 != nullptr && Map1 != nullptr && Map0->GetValueType() == Map1->GetValueType(), TEXT("Unexpected InnerJoin iteration result"));
 					},
-					Collection0.CreateConstSetMapIterator(),
-					Collection1.CreateSetMapIterator());
+					Collection0.CreateConstIterator(),
+					Collection1.CreateIterator());
 				AddErrorIfFalse(NumPredicateCalls == 5, TEXT("Unexpected InnerJoin iteration result"));
 
 				NumPredicateCalls = 0;
 				InnerJoinBy(
-					FAttributeCollectionSetMapJoinOp(),
+					FBoundMapCollectionJoinOp(),
 					[this, &NumPredicateCalls, &Collection0, &Collection1](FBoundValueMap* Map0, const FBoundValueMap* Map1)
 					{
 						NumPredicateCalls++;
@@ -236,93 +223,30 @@
 						AddErrorIfFalse(Map0 != nullptr && Map1 != nullptr && Map0->GetAttributeType() == Map1->GetAttributeType(), TEXT("Unexpected InnerJoin iteration result"));
 						AddErrorIfFalse(Map0 != nullptr && Map1 != nullptr && Map0->GetValueType() == Map1->GetValueType(), TEXT("Unexpected InnerJoin iteration result"));
 					},
-					Collection0.CreateSetMapIterator(),
-					Collection1.CreateConstSetMapIterator());
+					Collection0.CreateIterator(),
+					Collection1.CreateConstIterator());
 				AddErrorIfFalse(NumPredicateCalls == 5, TEXT("Unexpected InnerJoin iteration result"));
 			}
 
-			// Both inputs perfectly match (name/value maps)
-			{
-				int32 NumPredicateCalls = 0;
-				InnerJoinBy(
-					FAttributeCollectionMapJoinOp(),
-					[this, &NumPredicateCalls, &Collection0, &Collection1](FUnboundValueMap* Map0, FUnboundValueMap* Map1)
-					{
-						NumPredicateCalls++;
-						AddErrorIfFalse(CollectionContains(Collection0, Map0), TEXT("Unexpected InnerJoin iteration result"));
-						AddErrorIfFalse(CollectionContains(Collection1, Map1), TEXT("Unexpected InnerJoin iteration result"));
-						AddErrorIfFalse(Map0 != nullptr && Map1 != nullptr && Map0->GetValueType() == Map1->GetValueType(), TEXT("Unexpected InnerJoin iteration result"));
-					},
-					Collection0.CreateMapIterator(),
-					Collection1.CreateMapIterator());
-				AddErrorIfFalse(NumPredicateCalls == 2, TEXT("Unexpected InnerJoin iteration result"));
+			FBoundMapCollectionHeap Collection2(FullBodySet);
 
-				NumPredicateCalls = 0;
-				InnerJoinBy(
-					FAttributeCollectionMapJoinOp(),
-					[this, &NumPredicateCalls, &Collection0, &Collection1](const FUnboundValueMap* Map0, const FUnboundValueMap* Map1)
-					{
-						NumPredicateCalls++;
-						AddErrorIfFalse(CollectionContains(Collection0, Map0), TEXT("Unexpected InnerJoin iteration result"));
-						AddErrorIfFalse(CollectionContains(Collection1, Map1), TEXT("Unexpected InnerJoin iteration result"));
-						AddErrorIfFalse(Map0 != nullptr && Map1 != nullptr && Map0->GetValueType() == Map1->GetValueType(), TEXT("Unexpected InnerJoin iteration result"));
-					},
-					Collection0.CreateConstMapIterator(),
-					Collection1.CreateConstMapIterator());
-				AddErrorIfFalse(NumPredicateCalls == 2, TEXT("Unexpected InnerJoin iteration result"));
-
-				NumPredicateCalls = 0;
-				InnerJoinBy(
-					FAttributeCollectionMapJoinOp(),
-					[this, &NumPredicateCalls, &Collection0, &Collection1](const FUnboundValueMap* Map0, FUnboundValueMap* Map1)
-					{
-						NumPredicateCalls++;
-						AddErrorIfFalse(CollectionContains(Collection0, Map0), TEXT("Unexpected InnerJoin iteration result"));
-						AddErrorIfFalse(CollectionContains(Collection1, Map1), TEXT("Unexpected InnerJoin iteration result"));
-						AddErrorIfFalse(Map0 != nullptr && Map1 != nullptr && Map0->GetValueType() == Map1->GetValueType(), TEXT("Unexpected InnerJoin iteration result"));
-					},
-					Collection0.CreateConstMapIterator(),
-					Collection1.CreateMapIterator());
-				AddErrorIfFalse(NumPredicateCalls == 2, TEXT("Unexpected InnerJoin iteration result"));
-
-				NumPredicateCalls = 0;
-				InnerJoinBy(
-					FAttributeCollectionMapJoinOp(),
-					[this, &NumPredicateCalls, &Collection0, &Collection1](FUnboundValueMap* Map0, const FUnboundValueMap* Map1)
-					{
-						NumPredicateCalls++;
-						AddErrorIfFalse(CollectionContains(Collection0, Map0), TEXT("Unexpected InnerJoin iteration result"));
-						AddErrorIfFalse(CollectionContains(Collection1, Map1), TEXT("Unexpected InnerJoin iteration result"));
-						AddErrorIfFalse(Map0 != nullptr && Map1 != nullptr && Map0->GetValueType() == Map1->GetValueType(), TEXT("Unexpected InnerJoin iteration result"));
-					},
-					Collection0.CreateMapIterator(),
-					Collection1.CreateConstMapIterator());
-				AddErrorIfFalse(NumPredicateCalls == 2, TEXT("Unexpected InnerJoin iteration result"));
-			}
-
-			FAttributeCollectionHeap Collection2(FullBodySet);
-
-			AddErrorIfFalse(Collection2.IsEmpty(), TEXT("Attribute collection should be empty"));
+			AddErrorIfFalse(Collection2.IsEmpty(), TEXT("Collection should be empty"));
 			AddErrorIfFalse(Collection2.GetNamedSet() == FullBodySet, TEXT("Unexpected named set"));
 
-			TBoundValueMap<FFloatAnimationAttribute>* FloatToFloatMap2 = Collection2.AddSetMap<FFloatAnimationAttribute>(FloatToFloatKey);
+			TBoundValueMap<FFloatAnimationAttribute>* FloatToFloatMap2 = Collection2.Add<FFloatAnimationAttribute>(FloatToFloatKey);
 			UE_RETURN_ON_ERROR(FloatToFloatMap2 != nullptr, TEXT("Float->Float map should be valid"));
 
-			TBoundValueMap<FBoneTransformAnimationAttribute>* TransformToTransformMap2 = Collection2.AddSetMap<FBoneTransformAnimationAttribute>(TransformToTransformKey);
+			TBoundValueMap<FBoneTransformAnimationAttribute>* TransformToTransformMap2 = Collection2.Add<FBoneTransformAnimationAttribute>(TransformToTransformKey);
 			UE_RETURN_ON_ERROR(TransformToTransformMap2 != nullptr, TEXT("Transform->Transform map should be valid"));
 
-			TUnboundValueMap<FFloatAnimationAttribute>* FloatMap2 = Collection2.AddMap<FFloatAnimationAttribute>();
-			UE_RETURN_ON_ERROR(FloatMap2 != nullptr, TEXT("Float map should be valid"));
+			AddErrorIfFalse(!Collection2.IsEmpty(), TEXT("Collection should not be empty"));
+			AddErrorIfFalse(Collection2.Num() == 2, TEXT("Unexpected map count"));
 
-			AddErrorIfFalse(!Collection2.IsEmpty(), TEXT("Attribute collection should not be empty"));
-			AddErrorIfFalse(Collection2.NumSetMaps() == 2, TEXT("Unexpected set map count"));
-			AddErrorIfFalse(Collection2.NumMaps() == 1, TEXT("Unexpected map count"));
-
-			// First input is shorter and its entries are skipped when missing (set maps)
+			// First input is shorter and its entries are skipped when missing
 			{
 				int32 NumPredicateCalls = 0;
 				InnerJoinBy(
-					FAttributeCollectionSetMapJoinOp(),
+					FBoundMapCollectionJoinOp(),
 					[this, &NumPredicateCalls, &Collection0, &Collection2](FBoundValueMap* Map2, FBoundValueMap* Map0)
 					{
 						NumPredicateCalls++;
@@ -331,13 +255,13 @@
 						AddErrorIfFalse(Map0 != nullptr && Map2 != nullptr && Map0->GetAttributeType() == Map2->GetAttributeType(), TEXT("Unexpected InnerJoin iteration result"));
 						AddErrorIfFalse(Map0 != nullptr && Map2 != nullptr && Map0->GetValueType() == Map2->GetValueType(), TEXT("Unexpected InnerJoin iteration result"));
 					},
-					Collection2.CreateSetMapIterator(),
-					Collection0.CreateSetMapIterator());
+					Collection2.CreateIterator(),
+					Collection0.CreateIterator());
 				AddErrorIfFalse(NumPredicateCalls == 2, TEXT("Unexpected InnerJoin iteration result"));
 
 				NumPredicateCalls = 0;
 				InnerJoinBy(
-					FAttributeCollectionSetMapJoinOp(),
+					FBoundMapCollectionJoinOp(),
 					[this, &NumPredicateCalls, &Collection0, &Collection2](const FBoundValueMap* Map2, const FBoundValueMap* Map0)
 					{
 						NumPredicateCalls++;
@@ -346,13 +270,13 @@
 						AddErrorIfFalse(Map0 != nullptr && Map2 != nullptr && Map0->GetAttributeType() == Map2->GetAttributeType(), TEXT("Unexpected InnerJoin iteration result"));
 						AddErrorIfFalse(Map0 != nullptr && Map2 != nullptr && Map0->GetValueType() == Map2->GetValueType(), TEXT("Unexpected InnerJoin iteration result"));
 					},
-					Collection2.CreateConstSetMapIterator(),
-					Collection0.CreateConstSetMapIterator());
+					Collection2.CreateConstIterator(),
+					Collection0.CreateConstIterator());
 				AddErrorIfFalse(NumPredicateCalls == 2, TEXT("Unexpected InnerJoin iteration result"));
 
 				NumPredicateCalls = 0;
 				InnerJoinBy(
-					FAttributeCollectionSetMapJoinOp(),
+					FBoundMapCollectionJoinOp(),
 					[this, &NumPredicateCalls, &Collection0, &Collection2](const FBoundValueMap* Map2, FBoundValueMap* Map0)
 					{
 						NumPredicateCalls++;
@@ -361,13 +285,13 @@
 						AddErrorIfFalse(Map0 != nullptr && Map2 != nullptr && Map0->GetAttributeType() == Map2->GetAttributeType(), TEXT("Unexpected InnerJoin iteration result"));
 						AddErrorIfFalse(Map0 != nullptr && Map2 != nullptr && Map0->GetValueType() == Map2->GetValueType(), TEXT("Unexpected InnerJoin iteration result"));
 					},
-					Collection2.CreateConstSetMapIterator(),
-					Collection0.CreateSetMapIterator());
+					Collection2.CreateConstIterator(),
+					Collection0.CreateIterator());
 				AddErrorIfFalse(NumPredicateCalls == 2, TEXT("Unexpected InnerJoin iteration result"));
 
 				NumPredicateCalls = 0;
 				InnerJoinBy(
-					FAttributeCollectionSetMapJoinOp(),
+					FBoundMapCollectionJoinOp(),
 					[this, &NumPredicateCalls, &Collection0, &Collection2](FBoundValueMap* Map2, const FBoundValueMap* Map0)
 					{
 						NumPredicateCalls++;
@@ -376,75 +300,16 @@
 						AddErrorIfFalse(Map0 != nullptr && Map2 != nullptr && Map0->GetAttributeType() == Map2->GetAttributeType(), TEXT("Unexpected InnerJoin iteration result"));
 						AddErrorIfFalse(Map0 != nullptr && Map2 != nullptr && Map0->GetValueType() == Map2->GetValueType(), TEXT("Unexpected InnerJoin iteration result"));
 					},
-					Collection2.CreateSetMapIterator(),
-					Collection0.CreateConstSetMapIterator());
+					Collection2.CreateIterator(),
+					Collection0.CreateConstIterator());
 				AddErrorIfFalse(NumPredicateCalls == 2, TEXT("Unexpected InnerJoin iteration result"));
 			}
 
-			// First input is shorter and its entries are skipped when missing (name/value maps)
+			// Second input is shorter and its entries are skipped when missing
 			{
 				int32 NumPredicateCalls = 0;
 				InnerJoinBy(
-					FAttributeCollectionMapJoinOp(),
-					[this, &NumPredicateCalls, &Collection0, &Collection2](FUnboundValueMap* Map2, FUnboundValueMap* Map0)
-					{
-						NumPredicateCalls++;
-						AddErrorIfFalse(CollectionContains(Collection0, Map0), TEXT("Unexpected InnerJoin iteration result"));
-						AddErrorIfFalse(CollectionContains(Collection2, Map2), TEXT("Unexpected InnerJoin iteration result"));
-						AddErrorIfFalse(Map0 != nullptr && Map2 != nullptr && Map0->GetValueType() == Map2->GetValueType(), TEXT("Unexpected InnerJoin iteration result"));
-					},
-					Collection2.CreateMapIterator(),
-					Collection0.CreateMapIterator());
-				AddErrorIfFalse(NumPredicateCalls == 1, TEXT("Unexpected InnerJoin iteration result"));
-
-				NumPredicateCalls = 0;
-				InnerJoinBy(
-					FAttributeCollectionMapJoinOp(),
-					[this, &NumPredicateCalls, &Collection0, &Collection2](const FUnboundValueMap* Map2, const FUnboundValueMap* Map0)
-					{
-						NumPredicateCalls++;
-						AddErrorIfFalse(CollectionContains(Collection0, Map0), TEXT("Unexpected InnerJoin iteration result"));
-						AddErrorIfFalse(CollectionContains(Collection2, Map2), TEXT("Unexpected InnerJoin iteration result"));
-						AddErrorIfFalse(Map0 != nullptr && Map2 != nullptr && Map0->GetValueType() == Map2->GetValueType(), TEXT("Unexpected InnerJoin iteration result"));
-					},
-					Collection2.CreateConstMapIterator(),
-					Collection0.CreateConstMapIterator());
-				AddErrorIfFalse(NumPredicateCalls == 1, TEXT("Unexpected InnerJoin iteration result"));
-
-				NumPredicateCalls = 0;
-				InnerJoinBy(
-					FAttributeCollectionMapJoinOp(),
-					[this, &NumPredicateCalls, &Collection0, &Collection2](const FUnboundValueMap* Map2, FUnboundValueMap* Map0)
-					{
-						NumPredicateCalls++;
-						AddErrorIfFalse(CollectionContains(Collection0, Map0), TEXT("Unexpected InnerJoin iteration result"));
-						AddErrorIfFalse(CollectionContains(Collection2, Map2), TEXT("Unexpected InnerJoin iteration result"));
-						AddErrorIfFalse(Map0 != nullptr && Map2 != nullptr && Map0->GetValueType() == Map2->GetValueType(), TEXT("Unexpected InnerJoin iteration result"));
-					},
-					Collection2.CreateConstMapIterator(),
-					Collection0.CreateMapIterator());
-				AddErrorIfFalse(NumPredicateCalls == 1, TEXT("Unexpected InnerJoin iteration result"));
-
-				NumPredicateCalls = 0;
-				InnerJoinBy(
-					FAttributeCollectionMapJoinOp(),
-					[this, &NumPredicateCalls, &Collection0, &Collection2](FUnboundValueMap* Map2, const FUnboundValueMap* Map0)
-					{
-						NumPredicateCalls++;
-						AddErrorIfFalse(CollectionContains(Collection0, Map0), TEXT("Unexpected InnerJoin iteration result"));
-						AddErrorIfFalse(CollectionContains(Collection2, Map2), TEXT("Unexpected InnerJoin iteration result"));
-						AddErrorIfFalse(Map0 != nullptr && Map2 != nullptr && Map0->GetValueType() == Map2->GetValueType(), TEXT("Unexpected InnerJoin iteration result"));
-					},
-					Collection2.CreateMapIterator(),
-					Collection0.CreateConstMapIterator());
-				AddErrorIfFalse(NumPredicateCalls == 1, TEXT("Unexpected InnerJoin iteration result"));
-			}
-
-			// Second input is shorter and its entries are skipped when missing (set maps)
-			{
-				int32 NumPredicateCalls = 0;
-				InnerJoinBy(
-					FAttributeCollectionSetMapJoinOp(),
+					FBoundMapCollectionJoinOp(),
 					[this, &NumPredicateCalls, &Collection0, &Collection2](FBoundValueMap* Map0, FBoundValueMap* Map2)
 					{
 						NumPredicateCalls++;
@@ -453,13 +318,13 @@
 						AddErrorIfFalse(Map0 != nullptr && Map2 != nullptr && Map0->GetAttributeType() == Map2->GetAttributeType(), TEXT("Unexpected InnerJoin iteration result"));
 						AddErrorIfFalse(Map0 != nullptr && Map2 != nullptr && Map0->GetValueType() == Map2->GetValueType(), TEXT("Unexpected InnerJoin iteration result"));
 					},
-					Collection0.CreateSetMapIterator(),
-					Collection2.CreateSetMapIterator());
+					Collection0.CreateIterator(),
+					Collection2.CreateIterator());
 				AddErrorIfFalse(NumPredicateCalls == 2, TEXT("Unexpected InnerJoin iteration result"));
 
 				NumPredicateCalls = 0;
 				InnerJoinBy(
-					FAttributeCollectionSetMapJoinOp(),
+					FBoundMapCollectionJoinOp(),
 					[this, &NumPredicateCalls, &Collection0, &Collection2](const FBoundValueMap* Map0, const FBoundValueMap* Map2)
 					{
 						NumPredicateCalls++;
@@ -468,13 +333,13 @@
 						AddErrorIfFalse(Map0 != nullptr && Map2 != nullptr && Map0->GetAttributeType() == Map2->GetAttributeType(), TEXT("Unexpected InnerJoin iteration result"));
 						AddErrorIfFalse(Map0 != nullptr && Map2 != nullptr && Map0->GetValueType() == Map2->GetValueType(), TEXT("Unexpected InnerJoin iteration result"));
 					},
-					Collection0.CreateConstSetMapIterator(),
-					Collection2.CreateConstSetMapIterator());
+					Collection0.CreateConstIterator(),
+					Collection2.CreateConstIterator());
 				AddErrorIfFalse(NumPredicateCalls == 2, TEXT("Unexpected InnerJoin iteration result"));
 
 				NumPredicateCalls = 0;
 				InnerJoinBy(
-					FAttributeCollectionSetMapJoinOp(),
+					FBoundMapCollectionJoinOp(),
 					[this, &NumPredicateCalls, &Collection0, &Collection2](const FBoundValueMap* Map0, FBoundValueMap* Map2)
 					{
 						NumPredicateCalls++;
@@ -483,13 +348,13 @@
 						AddErrorIfFalse(Map0 != nullptr && Map2 != nullptr && Map0->GetAttributeType() == Map2->GetAttributeType(), TEXT("Unexpected InnerJoin iteration result"));
 						AddErrorIfFalse(Map0 != nullptr && Map2 != nullptr && Map0->GetValueType() == Map2->GetValueType(), TEXT("Unexpected InnerJoin iteration result"));
 					},
-					Collection0.CreateConstSetMapIterator(),
-					Collection2.CreateSetMapIterator());
+					Collection0.CreateConstIterator(),
+					Collection2.CreateIterator());
 				AddErrorIfFalse(NumPredicateCalls == 2, TEXT("Unexpected InnerJoin iteration result"));
 
 				NumPredicateCalls = 0;
 				InnerJoinBy(
-					FAttributeCollectionSetMapJoinOp(),
+					FBoundMapCollectionJoinOp(),
 					[this, &NumPredicateCalls, &Collection0, &Collection2](FBoundValueMap* Map0, const FBoundValueMap* Map2)
 					{
 						NumPredicateCalls++;
@@ -498,75 +363,16 @@
 						AddErrorIfFalse(Map0 != nullptr && Map2 != nullptr && Map0->GetAttributeType() == Map2->GetAttributeType(), TEXT("Unexpected InnerJoin iteration result"));
 						AddErrorIfFalse(Map0 != nullptr && Map2 != nullptr && Map0->GetValueType() == Map2->GetValueType(), TEXT("Unexpected InnerJoin iteration result"));
 					},
-					Collection0.CreateSetMapIterator(),
-					Collection2.CreateConstSetMapIterator());
+					Collection0.CreateIterator(),
+					Collection2.CreateConstIterator());
 				AddErrorIfFalse(NumPredicateCalls == 2, TEXT("Unexpected InnerJoin iteration result"));
 			}
 
-			// Second input is shorter and its entries are skipped when missing (name/value maps)
+			// 3 input variant
 			{
 				int32 NumPredicateCalls = 0;
 				InnerJoinBy(
-					FAttributeCollectionMapJoinOp(),
-					[this, &NumPredicateCalls, &Collection0, &Collection2](FUnboundValueMap* Map0, FUnboundValueMap* Map2)
-					{
-						NumPredicateCalls++;
-						AddErrorIfFalse(CollectionContains(Collection0, Map0), TEXT("Unexpected InnerJoin iteration result"));
-						AddErrorIfFalse(CollectionContains(Collection2, Map2), TEXT("Unexpected InnerJoin iteration result"));
-						AddErrorIfFalse(Map0 != nullptr && Map2 != nullptr && Map0->GetValueType() == Map2->GetValueType(), TEXT("Unexpected InnerJoin iteration result"));
-					},
-					Collection0.CreateMapIterator(),
-					Collection2.CreateMapIterator());
-				AddErrorIfFalse(NumPredicateCalls == 1, TEXT("Unexpected InnerJoin iteration result"));
-
-				NumPredicateCalls = 0;
-				InnerJoinBy(
-					FAttributeCollectionMapJoinOp(),
-					[this, &NumPredicateCalls, &Collection0, &Collection2](const FUnboundValueMap* Map0, const FUnboundValueMap* Map2)
-					{
-						NumPredicateCalls++;
-						AddErrorIfFalse(CollectionContains(Collection0, Map0), TEXT("Unexpected InnerJoin iteration result"));
-						AddErrorIfFalse(CollectionContains(Collection2, Map2), TEXT("Unexpected InnerJoin iteration result"));
-						AddErrorIfFalse(Map0 != nullptr && Map2 != nullptr && Map0->GetValueType() == Map2->GetValueType(), TEXT("Unexpected InnerJoin iteration result"));
-					},
-					Collection0.CreateConstMapIterator(),
-					Collection2.CreateConstMapIterator());
-				AddErrorIfFalse(NumPredicateCalls == 1, TEXT("Unexpected InnerJoin iteration result"));
-
-				NumPredicateCalls = 0;
-				InnerJoinBy(
-					FAttributeCollectionMapJoinOp(),
-					[this, &NumPredicateCalls, &Collection0, &Collection2](const FUnboundValueMap* Map0, FUnboundValueMap* Map2)
-					{
-						NumPredicateCalls++;
-						AddErrorIfFalse(CollectionContains(Collection0, Map0), TEXT("Unexpected InnerJoin iteration result"));
-						AddErrorIfFalse(CollectionContains(Collection2, Map2), TEXT("Unexpected InnerJoin iteration result"));
-						AddErrorIfFalse(Map0 != nullptr && Map2 != nullptr && Map0->GetValueType() == Map2->GetValueType(), TEXT("Unexpected InnerJoin iteration result"));
-					},
-					Collection0.CreateConstMapIterator(),
-					Collection2.CreateMapIterator());
-				AddErrorIfFalse(NumPredicateCalls == 1, TEXT("Unexpected InnerJoin iteration result"));
-
-				NumPredicateCalls = 0;
-				InnerJoinBy(
-					FAttributeCollectionMapJoinOp(),
-					[this, &NumPredicateCalls, &Collection0, &Collection2](FUnboundValueMap* Map0, const FUnboundValueMap* Map2)
-					{
-						NumPredicateCalls++;
-						AddErrorIfFalse(CollectionContains(Collection0, Map0), TEXT("Unexpected InnerJoin iteration result"));
-						AddErrorIfFalse(CollectionContains(Collection2, Map2), TEXT("Unexpected InnerJoin iteration result"));
-						AddErrorIfFalse(Map0 != nullptr && Map2 != nullptr && Map0->GetValueType() == Map2->GetValueType(), TEXT("Unexpected InnerJoin iteration result"));
-					},
-					Collection0.CreateMapIterator(),
-					Collection2.CreateConstMapIterator());
-				AddErrorIfFalse(NumPredicateCalls == 1, TEXT("Unexpected InnerJoin iteration result"));
-			}
-
-			// 3 input variant (set maps)
-			{
-				int32 NumPredicateCalls = 0;
-				InnerJoinBy(
-					FAttributeCollectionSetMapJoinOp(),
+					FBoundMapCollectionJoinOp(),
 					[this, &NumPredicateCalls, &Collection0, &Collection1, &Collection2](FBoundValueMap* Map0, FBoundValueMap* Map1, FBoundValueMap* Map2)
 					{
 						NumPredicateCalls++;
@@ -578,14 +384,14 @@
 						AddErrorIfFalse(Map0 != nullptr && Map1 != nullptr && Map0->GetValueType() == Map1->GetValueType(), TEXT("Unexpected InnerJoin iteration result"));
 						AddErrorIfFalse(Map0 != nullptr && Map2 != nullptr && Map0->GetValueType() == Map2->GetValueType(), TEXT("Unexpected InnerJoin iteration result"));
 					},
-					Collection0.CreateSetMapIterator(),
-					Collection1.CreateSetMapIterator(),
-					Collection2.CreateSetMapIterator());
+					Collection0.CreateIterator(),
+					Collection1.CreateIterator(),
+					Collection2.CreateIterator());
 				AddErrorIfFalse(NumPredicateCalls == 2, TEXT("Unexpected InnerJoin iteration result"));
 
 				NumPredicateCalls = 0;
 				InnerJoinBy(
-					FAttributeCollectionSetMapJoinOp(),
+					FBoundMapCollectionJoinOp(),
 					[this, &NumPredicateCalls, &Collection0, &Collection1, &Collection2](FBoundValueMap* Map0, FBoundValueMap* Map2, FBoundValueMap* Map1)
 					{
 						NumPredicateCalls++;
@@ -597,14 +403,14 @@
 						AddErrorIfFalse(Map0 != nullptr && Map1 != nullptr && Map0->GetValueType() == Map1->GetValueType(), TEXT("Unexpected InnerJoin iteration result"));
 						AddErrorIfFalse(Map0 != nullptr && Map2 != nullptr && Map0->GetValueType() == Map2->GetValueType(), TEXT("Unexpected InnerJoin iteration result"));
 					},
-					Collection0.CreateSetMapIterator(),
-					Collection2.CreateSetMapIterator(),
-					Collection1.CreateSetMapIterator());
+					Collection0.CreateIterator(),
+					Collection2.CreateIterator(),
+					Collection1.CreateIterator());
 				AddErrorIfFalse(NumPredicateCalls == 2, TEXT("Unexpected InnerJoin iteration result"));
 
 				NumPredicateCalls = 0;
 				InnerJoinBy(
-					FAttributeCollectionSetMapJoinOp(),
+					FBoundMapCollectionJoinOp(),
 					[this, &NumPredicateCalls, &Collection0, &Collection1, &Collection2](FBoundValueMap* Map1, FBoundValueMap* Map0, FBoundValueMap* Map2)
 					{
 						NumPredicateCalls++;
@@ -616,14 +422,14 @@
 						AddErrorIfFalse(Map0 != nullptr && Map1 != nullptr && Map0->GetValueType() == Map1->GetValueType(), TEXT("Unexpected InnerJoin iteration result"));
 						AddErrorIfFalse(Map0 != nullptr && Map2 != nullptr && Map0->GetValueType() == Map2->GetValueType(), TEXT("Unexpected InnerJoin iteration result"));
 					},
-					Collection1.CreateSetMapIterator(),
-					Collection0.CreateSetMapIterator(),
-					Collection2.CreateSetMapIterator());
+					Collection1.CreateIterator(),
+					Collection0.CreateIterator(),
+					Collection2.CreateIterator());
 				AddErrorIfFalse(NumPredicateCalls == 2, TEXT("Unexpected InnerJoin iteration result"));
 
 				NumPredicateCalls = 0;
 				InnerJoinBy(
-					FAttributeCollectionSetMapJoinOp(),
+					FBoundMapCollectionJoinOp(),
 					[this, &NumPredicateCalls, &Collection0, &Collection1, &Collection2](FBoundValueMap* Map1, FBoundValueMap* Map2, FBoundValueMap* Map0)
 					{
 						NumPredicateCalls++;
@@ -635,14 +441,14 @@
 						AddErrorIfFalse(Map0 != nullptr && Map1 != nullptr && Map0->GetValueType() == Map1->GetValueType(), TEXT("Unexpected InnerJoin iteration result"));
 						AddErrorIfFalse(Map0 != nullptr && Map2 != nullptr && Map0->GetValueType() == Map2->GetValueType(), TEXT("Unexpected InnerJoin iteration result"));
 					},
-					Collection1.CreateSetMapIterator(),
-					Collection2.CreateSetMapIterator(),
-					Collection0.CreateSetMapIterator());
+					Collection1.CreateIterator(),
+					Collection2.CreateIterator(),
+					Collection0.CreateIterator());
 				AddErrorIfFalse(NumPredicateCalls == 2, TEXT("Unexpected InnerJoin iteration result"));
 
 				NumPredicateCalls = 0;
 				InnerJoinBy(
-					FAttributeCollectionSetMapJoinOp(),
+					FBoundMapCollectionJoinOp(),
 					[this, &NumPredicateCalls, &Collection0, &Collection1, &Collection2](FBoundValueMap* Map2, FBoundValueMap* Map0, FBoundValueMap* Map1)
 					{
 						NumPredicateCalls++;
@@ -654,14 +460,14 @@
 						AddErrorIfFalse(Map0 != nullptr && Map1 != nullptr && Map0->GetValueType() == Map1->GetValueType(), TEXT("Unexpected InnerJoin iteration result"));
 						AddErrorIfFalse(Map0 != nullptr && Map2 != nullptr && Map0->GetValueType() == Map2->GetValueType(), TEXT("Unexpected InnerJoin iteration result"));
 					},
-					Collection2.CreateSetMapIterator(),
-					Collection0.CreateSetMapIterator(),
-					Collection1.CreateSetMapIterator());
+					Collection2.CreateIterator(),
+					Collection0.CreateIterator(),
+					Collection1.CreateIterator());
 				AddErrorIfFalse(NumPredicateCalls == 2, TEXT("Unexpected InnerJoin iteration result"));
 
 				NumPredicateCalls = 0;
 				InnerJoinBy(
-					FAttributeCollectionSetMapJoinOp(),
+					FBoundMapCollectionJoinOp(),
 					[this, &NumPredicateCalls, &Collection0, &Collection1, &Collection2](FBoundValueMap* Map2, FBoundValueMap* Map1, FBoundValueMap* Map0)
 					{
 						NumPredicateCalls++;
@@ -673,118 +479,13 @@
 						AddErrorIfFalse(Map0 != nullptr && Map1 != nullptr && Map0->GetValueType() == Map1->GetValueType(), TEXT("Unexpected InnerJoin iteration result"));
 						AddErrorIfFalse(Map0 != nullptr && Map2 != nullptr && Map0->GetValueType() == Map2->GetValueType(), TEXT("Unexpected InnerJoin iteration result"));
 					},
-					Collection2.CreateSetMapIterator(),
-					Collection1.CreateSetMapIterator(),
-					Collection0.CreateSetMapIterator());
+					Collection2.CreateIterator(),
+					Collection1.CreateIterator(),
+					Collection0.CreateIterator());
 				AddErrorIfFalse(NumPredicateCalls == 2, TEXT("Unexpected InnerJoin iteration result"));
 			}
 
-			// 3 input variant (name/value maps)
-			{
-				int32 NumPredicateCalls = 0;
-				InnerJoinBy(
-					FAttributeCollectionMapJoinOp(),
-					[this, &NumPredicateCalls, &Collection0, &Collection1, &Collection2](FUnboundValueMap* Map0, FUnboundValueMap* Map1, FUnboundValueMap* Map2)
-					{
-						NumPredicateCalls++;
-						AddErrorIfFalse(CollectionContains(Collection0, Map0), TEXT("Unexpected InnerJoin iteration result"));
-						AddErrorIfFalse(CollectionContains(Collection1, Map1), TEXT("Unexpected InnerJoin iteration result"));
-						AddErrorIfFalse(CollectionContains(Collection2, Map2), TEXT("Unexpected InnerJoin iteration result"));
-						AddErrorIfFalse(Map0 != nullptr && Map1 != nullptr && Map0->GetValueType() == Map1->GetValueType(), TEXT("Unexpected InnerJoin iteration result"));
-						AddErrorIfFalse(Map0 != nullptr && Map2 != nullptr && Map0->GetValueType() == Map2->GetValueType(), TEXT("Unexpected InnerJoin iteration result"));
-					},
-					Collection0.CreateMapIterator(),
-					Collection1.CreateMapIterator(),
-					Collection2.CreateMapIterator());
-				AddErrorIfFalse(NumPredicateCalls == 1, TEXT("Unexpected InnerJoin iteration result"));
-
-				NumPredicateCalls = 0;
-				InnerJoinBy(
-					FAttributeCollectionMapJoinOp(),
-					[this, &NumPredicateCalls, &Collection0, &Collection1, &Collection2](FUnboundValueMap* Map0, FUnboundValueMap* Map2, FUnboundValueMap* Map1)
-					{
-						NumPredicateCalls++;
-						AddErrorIfFalse(CollectionContains(Collection0, Map0), TEXT("Unexpected InnerJoin iteration result"));
-						AddErrorIfFalse(CollectionContains(Collection1, Map1), TEXT("Unexpected InnerJoin iteration result"));
-						AddErrorIfFalse(CollectionContains(Collection2, Map2), TEXT("Unexpected InnerJoin iteration result"));
-						AddErrorIfFalse(Map0 != nullptr && Map1 != nullptr && Map0->GetValueType() == Map1->GetValueType(), TEXT("Unexpected InnerJoin iteration result"));
-						AddErrorIfFalse(Map0 != nullptr && Map2 != nullptr && Map0->GetValueType() == Map2->GetValueType(), TEXT("Unexpected InnerJoin iteration result"));
-					},
-					Collection0.CreateMapIterator(),
-					Collection2.CreateMapIterator(),
-					Collection1.CreateMapIterator());
-				AddErrorIfFalse(NumPredicateCalls == 1, TEXT("Unexpected InnerJoin iteration result"));
-
-				NumPredicateCalls = 0;
-				InnerJoinBy(
-					FAttributeCollectionMapJoinOp(),
-					[this, &NumPredicateCalls, &Collection0, &Collection1, &Collection2](FUnboundValueMap* Map1, FUnboundValueMap* Map0, FUnboundValueMap* Map2)
-					{
-						NumPredicateCalls++;
-						AddErrorIfFalse(CollectionContains(Collection0, Map0), TEXT("Unexpected InnerJoin iteration result"));
-						AddErrorIfFalse(CollectionContains(Collection1, Map1), TEXT("Unexpected InnerJoin iteration result"));
-						AddErrorIfFalse(CollectionContains(Collection2, Map2), TEXT("Unexpected InnerJoin iteration result"));
-						AddErrorIfFalse(Map0 != nullptr && Map1 != nullptr && Map0->GetValueType() == Map1->GetValueType(), TEXT("Unexpected InnerJoin iteration result"));
-						AddErrorIfFalse(Map0 != nullptr && Map2 != nullptr && Map0->GetValueType() == Map2->GetValueType(), TEXT("Unexpected InnerJoin iteration result"));
-					},
-					Collection1.CreateMapIterator(),
-					Collection0.CreateMapIterator(),
-					Collection2.CreateMapIterator());
-				AddErrorIfFalse(NumPredicateCalls == 1, TEXT("Unexpected InnerJoin iteration result"));
-
-				NumPredicateCalls = 0;
-				InnerJoinBy(
-					FAttributeCollectionMapJoinOp(),
-					[this, &NumPredicateCalls, &Collection0, &Collection1, &Collection2](FUnboundValueMap* Map1, FUnboundValueMap* Map2, FUnboundValueMap* Map0)
-					{
-						NumPredicateCalls++;
-						AddErrorIfFalse(CollectionContains(Collection0, Map0), TEXT("Unexpected InnerJoin iteration result"));
-						AddErrorIfFalse(CollectionContains(Collection1, Map1), TEXT("Unexpected InnerJoin iteration result"));
-						AddErrorIfFalse(CollectionContains(Collection2, Map2), TEXT("Unexpected InnerJoin iteration result"));
-						AddErrorIfFalse(Map0 != nullptr && Map1 != nullptr && Map0->GetValueType() == Map1->GetValueType(), TEXT("Unexpected InnerJoin iteration result"));
-						AddErrorIfFalse(Map0 != nullptr && Map2 != nullptr && Map0->GetValueType() == Map2->GetValueType(), TEXT("Unexpected InnerJoin iteration result"));
-					},
-					Collection1.CreateMapIterator(),
-					Collection2.CreateMapIterator(),
-					Collection0.CreateMapIterator());
-				AddErrorIfFalse(NumPredicateCalls == 1, TEXT("Unexpected InnerJoin iteration result"));
-
-				NumPredicateCalls = 0;
-				InnerJoinBy(
-					FAttributeCollectionMapJoinOp(),
-					[this, &NumPredicateCalls, &Collection0, &Collection1, &Collection2](FUnboundValueMap* Map2, FUnboundValueMap* Map0, FUnboundValueMap* Map1)
-					{
-						NumPredicateCalls++;
-						AddErrorIfFalse(CollectionContains(Collection0, Map0), TEXT("Unexpected InnerJoin iteration result"));
-						AddErrorIfFalse(CollectionContains(Collection1, Map1), TEXT("Unexpected InnerJoin iteration result"));
-						AddErrorIfFalse(CollectionContains(Collection2, Map2), TEXT("Unexpected InnerJoin iteration result"));
-						AddErrorIfFalse(Map0 != nullptr && Map1 != nullptr && Map0->GetValueType() == Map1->GetValueType(), TEXT("Unexpected InnerJoin iteration result"));
-						AddErrorIfFalse(Map0 != nullptr && Map2 != nullptr && Map0->GetValueType() == Map2->GetValueType(), TEXT("Unexpected InnerJoin iteration result"));
-					},
-					Collection2.CreateMapIterator(),
-					Collection0.CreateMapIterator(),
-					Collection1.CreateMapIterator());
-				AddErrorIfFalse(NumPredicateCalls == 1, TEXT("Unexpected InnerJoin iteration result"));
-
-				NumPredicateCalls = 0;
-				InnerJoinBy(
-					FAttributeCollectionMapJoinOp(),
-					[this, &NumPredicateCalls, &Collection0, &Collection1, &Collection2](FUnboundValueMap* Map2, FUnboundValueMap* Map1, FUnboundValueMap* Map0)
-					{
-						NumPredicateCalls++;
-						AddErrorIfFalse(CollectionContains(Collection0, Map0), TEXT("Unexpected InnerJoin iteration result"));
-						AddErrorIfFalse(CollectionContains(Collection1, Map1), TEXT("Unexpected InnerJoin iteration result"));
-						AddErrorIfFalse(CollectionContains(Collection2, Map2), TEXT("Unexpected InnerJoin iteration result"));
-						AddErrorIfFalse(Map0 != nullptr && Map1 != nullptr && Map0->GetValueType() == Map1->GetValueType(), TEXT("Unexpected InnerJoin iteration result"));
-						AddErrorIfFalse(Map0 != nullptr && Map2 != nullptr && Map0->GetValueType() == Map2->GetValueType(), TEXT("Unexpected InnerJoin iteration result"));
-					},
-					Collection2.CreateMapIterator(),
-					Collection1.CreateMapIterator(),
-					Collection0.CreateMapIterator());
-				AddErrorIfFalse(NumPredicateCalls == 1, TEXT("Unexpected InnerJoin iteration result"));
-			}
-
-			// Join with a transformer list (set maps)
+			// Join with a transformer list
 			{
 				// Use a float transformer
 				// We set up Collection0 such that it has two mappings with float values: Transform -> Float and Float -> Float
@@ -809,8 +510,8 @@
 				};
 
 #ifdef _MSC_VER
-	#pragma warning(push)
-	#pragma warning(disable : 4191) // unsafe conversion from 'function pointer type A' to 'function pointer type B'
+#pragma warning(push)
+#pragma warning(disable : 4191) // unsafe conversion from 'function pointer type A' to 'function pointer type B'
 #endif
 
 				FValueTransformerList TransformerList(Transformers::FSanitize::TransformerName);
@@ -818,7 +519,7 @@
 				TransformerList.AddBoundValueMapTransformer(FVectorAnimationAttribute::StaticStruct(), reinterpret_cast<FRawTransformerFunc>(&FVectorAttributeTransformer_SanitizeTest::TransformBoundValueMap));
 
 #ifdef _MSC_VER
-	#pragma warning(pop)
+#pragma warning(pop)
 #endif
 
 				int32 NumPredicateCalls = 0;
@@ -831,11 +532,332 @@
 						AddErrorIfFalse(SetMap->GetValueType() == FFloatAnimationAttribute::StaticStruct(), TEXT("Unexpected InnerJoin iteration result"));
 					},
 					TransformerList.CreateBoundValueMapTransformerIterator(),
-					Collection0.CreateSetMapIterator());
+					Collection0.CreateIterator());
+				AddErrorIfFalse(NumPredicateCalls == 2, TEXT("Unexpected InnerJoin iteration result"));
+			}
+		}
+
+		// Unbound value map collection
+		{
+			FUnboundMapCollectionHeap Collection0;
+
+			AddErrorIfFalse(Collection0.IsEmpty(), TEXT("Collection should be empty"));
+
+			TUnboundValueMap<FFloatAnimationAttribute>* FloatMap0 = Collection0.Add<FFloatAnimationAttribute>();
+			UE_RETURN_ON_ERROR(FloatMap0 != nullptr, TEXT("Float map should be valid"));
+
+			TUnboundValueMap<FBoneTransformAnimationAttribute>* BoneMap0 = Collection0.Add<FBoneTransformAnimationAttribute>();
+			UE_RETURN_ON_ERROR(BoneMap0 != nullptr, TEXT("Bone map should be valid"));
+
+			AddErrorIfFalse(!Collection0.IsEmpty(), TEXT("Collection should not be empty"));
+			AddErrorIfFalse(Collection0.Num() == 2, TEXT("Unexpected map count"));
+
+			FUnboundMapCollectionHeap Collection1;
+
+			AddErrorIfFalse(Collection1.IsEmpty(), TEXT("Collection should be empty"));
+
+			TUnboundValueMap<FFloatAnimationAttribute>* FloatMap1 = Collection1.Add<FFloatAnimationAttribute>();
+			UE_RETURN_ON_ERROR(FloatMap1 != nullptr, TEXT("Float map should be valid"));
+
+			TUnboundValueMap<FBoneTransformAnimationAttribute>* BoneMap1 = Collection1.Add<FBoneTransformAnimationAttribute>();
+			UE_RETURN_ON_ERROR(BoneMap1 != nullptr, TEXT("Bone map should be valid"));
+
+			AddErrorIfFalse(!Collection1.IsEmpty(), TEXT("Collection should not be empty"));
+			AddErrorIfFalse(Collection1.Num() == 2, TEXT("Unexpected map count"));
+
+			// Both inputs perfectly match
+			{
+				int32 NumPredicateCalls = 0;
+				InnerJoinBy(
+					FUnboundMapCollectionJoinOp(),
+					[this, &NumPredicateCalls, &Collection0, &Collection1](FUnboundValueMap* Map0, FUnboundValueMap* Map1)
+					{
+						NumPredicateCalls++;
+						AddErrorIfFalse(CollectionContains(Collection0, Map0), TEXT("Unexpected InnerJoin iteration result"));
+						AddErrorIfFalse(CollectionContains(Collection1, Map1), TEXT("Unexpected InnerJoin iteration result"));
+						AddErrorIfFalse(Map0 != nullptr && Map1 != nullptr && Map0->GetValueType() == Map1->GetValueType(), TEXT("Unexpected InnerJoin iteration result"));
+					},
+					Collection0.CreateIterator(),
+					Collection1.CreateIterator());
+				AddErrorIfFalse(NumPredicateCalls == 2, TEXT("Unexpected InnerJoin iteration result"));
+
+				NumPredicateCalls = 0;
+				InnerJoinBy(
+					FUnboundMapCollectionJoinOp(),
+					[this, &NumPredicateCalls, &Collection0, &Collection1](const FUnboundValueMap* Map0, const FUnboundValueMap* Map1)
+					{
+						NumPredicateCalls++;
+						AddErrorIfFalse(CollectionContains(Collection0, Map0), TEXT("Unexpected InnerJoin iteration result"));
+						AddErrorIfFalse(CollectionContains(Collection1, Map1), TEXT("Unexpected InnerJoin iteration result"));
+						AddErrorIfFalse(Map0 != nullptr && Map1 != nullptr && Map0->GetValueType() == Map1->GetValueType(), TEXT("Unexpected InnerJoin iteration result"));
+					},
+					Collection0.CreateConstIterator(),
+					Collection1.CreateConstIterator());
+				AddErrorIfFalse(NumPredicateCalls == 2, TEXT("Unexpected InnerJoin iteration result"));
+
+				NumPredicateCalls = 0;
+				InnerJoinBy(
+					FUnboundMapCollectionJoinOp(),
+					[this, &NumPredicateCalls, &Collection0, &Collection1](const FUnboundValueMap* Map0, FUnboundValueMap* Map1)
+					{
+						NumPredicateCalls++;
+						AddErrorIfFalse(CollectionContains(Collection0, Map0), TEXT("Unexpected InnerJoin iteration result"));
+						AddErrorIfFalse(CollectionContains(Collection1, Map1), TEXT("Unexpected InnerJoin iteration result"));
+						AddErrorIfFalse(Map0 != nullptr && Map1 != nullptr && Map0->GetValueType() == Map1->GetValueType(), TEXT("Unexpected InnerJoin iteration result"));
+					},
+					Collection0.CreateConstIterator(),
+					Collection1.CreateIterator());
+				AddErrorIfFalse(NumPredicateCalls == 2, TEXT("Unexpected InnerJoin iteration result"));
+
+				NumPredicateCalls = 0;
+				InnerJoinBy(
+					FUnboundMapCollectionJoinOp(),
+					[this, &NumPredicateCalls, &Collection0, &Collection1](FUnboundValueMap* Map0, const FUnboundValueMap* Map1)
+					{
+						NumPredicateCalls++;
+						AddErrorIfFalse(CollectionContains(Collection0, Map0), TEXT("Unexpected InnerJoin iteration result"));
+						AddErrorIfFalse(CollectionContains(Collection1, Map1), TEXT("Unexpected InnerJoin iteration result"));
+						AddErrorIfFalse(Map0 != nullptr && Map1 != nullptr && Map0->GetValueType() == Map1->GetValueType(), TEXT("Unexpected InnerJoin iteration result"));
+					},
+					Collection0.CreateIterator(),
+					Collection1.CreateConstIterator());
 				AddErrorIfFalse(NumPredicateCalls == 2, TEXT("Unexpected InnerJoin iteration result"));
 			}
 
-			// Join with a transformer list (name/value maps)
+			FUnboundMapCollectionHeap Collection2;
+
+			AddErrorIfFalse(Collection2.IsEmpty(), TEXT("Collection should be empty"));
+
+			TUnboundValueMap<FFloatAnimationAttribute>* FloatMap2 = Collection2.Add<FFloatAnimationAttribute>();
+			UE_RETURN_ON_ERROR(FloatMap2 != nullptr, TEXT("Float map should be valid"));
+
+			AddErrorIfFalse(!Collection2.IsEmpty(), TEXT("Collection should not be empty"));
+			AddErrorIfFalse(Collection2.Num() == 1, TEXT("Unexpected map count"));
+
+			// First input is shorter and its entries are skipped when missing
+			{
+				int32 NumPredicateCalls = 0;
+				InnerJoinBy(
+					FUnboundMapCollectionJoinOp(),
+					[this, &NumPredicateCalls, &Collection0, &Collection2](FUnboundValueMap* Map2, FUnboundValueMap* Map0)
+					{
+						NumPredicateCalls++;
+						AddErrorIfFalse(CollectionContains(Collection0, Map0), TEXT("Unexpected InnerJoin iteration result"));
+						AddErrorIfFalse(CollectionContains(Collection2, Map2), TEXT("Unexpected InnerJoin iteration result"));
+						AddErrorIfFalse(Map0 != nullptr && Map2 != nullptr && Map0->GetValueType() == Map2->GetValueType(), TEXT("Unexpected InnerJoin iteration result"));
+					},
+					Collection2.CreateIterator(),
+					Collection0.CreateIterator());
+				AddErrorIfFalse(NumPredicateCalls == 1, TEXT("Unexpected InnerJoin iteration result"));
+
+				NumPredicateCalls = 0;
+				InnerJoinBy(
+					FUnboundMapCollectionJoinOp(),
+					[this, &NumPredicateCalls, &Collection0, &Collection2](const FUnboundValueMap* Map2, const FUnboundValueMap* Map0)
+					{
+						NumPredicateCalls++;
+						AddErrorIfFalse(CollectionContains(Collection0, Map0), TEXT("Unexpected InnerJoin iteration result"));
+						AddErrorIfFalse(CollectionContains(Collection2, Map2), TEXT("Unexpected InnerJoin iteration result"));
+						AddErrorIfFalse(Map0 != nullptr && Map2 != nullptr && Map0->GetValueType() == Map2->GetValueType(), TEXT("Unexpected InnerJoin iteration result"));
+					},
+					Collection2.CreateConstIterator(),
+					Collection0.CreateConstIterator());
+				AddErrorIfFalse(NumPredicateCalls == 1, TEXT("Unexpected InnerJoin iteration result"));
+
+				NumPredicateCalls = 0;
+				InnerJoinBy(
+					FUnboundMapCollectionJoinOp(),
+					[this, &NumPredicateCalls, &Collection0, &Collection2](const FUnboundValueMap* Map2, FUnboundValueMap* Map0)
+					{
+						NumPredicateCalls++;
+						AddErrorIfFalse(CollectionContains(Collection0, Map0), TEXT("Unexpected InnerJoin iteration result"));
+						AddErrorIfFalse(CollectionContains(Collection2, Map2), TEXT("Unexpected InnerJoin iteration result"));
+						AddErrorIfFalse(Map0 != nullptr && Map2 != nullptr && Map0->GetValueType() == Map2->GetValueType(), TEXT("Unexpected InnerJoin iteration result"));
+					},
+					Collection2.CreateConstIterator(),
+					Collection0.CreateIterator());
+				AddErrorIfFalse(NumPredicateCalls == 1, TEXT("Unexpected InnerJoin iteration result"));
+
+				NumPredicateCalls = 0;
+				InnerJoinBy(
+					FUnboundMapCollectionJoinOp(),
+					[this, &NumPredicateCalls, &Collection0, &Collection2](FUnboundValueMap* Map2, const FUnboundValueMap* Map0)
+					{
+						NumPredicateCalls++;
+						AddErrorIfFalse(CollectionContains(Collection0, Map0), TEXT("Unexpected InnerJoin iteration result"));
+						AddErrorIfFalse(CollectionContains(Collection2, Map2), TEXT("Unexpected InnerJoin iteration result"));
+						AddErrorIfFalse(Map0 != nullptr && Map2 != nullptr && Map0->GetValueType() == Map2->GetValueType(), TEXT("Unexpected InnerJoin iteration result"));
+					},
+					Collection2.CreateIterator(),
+					Collection0.CreateConstIterator());
+				AddErrorIfFalse(NumPredicateCalls == 1, TEXT("Unexpected InnerJoin iteration result"));
+			}
+
+			// Second input is shorter and its entries are skipped when missing
+			{
+				int32 NumPredicateCalls = 0;
+				InnerJoinBy(
+					FUnboundMapCollectionJoinOp(),
+					[this, &NumPredicateCalls, &Collection0, &Collection2](FUnboundValueMap* Map0, FUnboundValueMap* Map2)
+					{
+						NumPredicateCalls++;
+						AddErrorIfFalse(CollectionContains(Collection0, Map0), TEXT("Unexpected InnerJoin iteration result"));
+						AddErrorIfFalse(CollectionContains(Collection2, Map2), TEXT("Unexpected InnerJoin iteration result"));
+						AddErrorIfFalse(Map0 != nullptr && Map2 != nullptr && Map0->GetValueType() == Map2->GetValueType(), TEXT("Unexpected InnerJoin iteration result"));
+					},
+					Collection0.CreateIterator(),
+					Collection2.CreateIterator());
+				AddErrorIfFalse(NumPredicateCalls == 1, TEXT("Unexpected InnerJoin iteration result"));
+
+				NumPredicateCalls = 0;
+				InnerJoinBy(
+					FUnboundMapCollectionJoinOp(),
+					[this, &NumPredicateCalls, &Collection0, &Collection2](const FUnboundValueMap* Map0, const FUnboundValueMap* Map2)
+					{
+						NumPredicateCalls++;
+						AddErrorIfFalse(CollectionContains(Collection0, Map0), TEXT("Unexpected InnerJoin iteration result"));
+						AddErrorIfFalse(CollectionContains(Collection2, Map2), TEXT("Unexpected InnerJoin iteration result"));
+						AddErrorIfFalse(Map0 != nullptr && Map2 != nullptr && Map0->GetValueType() == Map2->GetValueType(), TEXT("Unexpected InnerJoin iteration result"));
+					},
+					Collection0.CreateConstIterator(),
+					Collection2.CreateConstIterator());
+				AddErrorIfFalse(NumPredicateCalls == 1, TEXT("Unexpected InnerJoin iteration result"));
+
+				NumPredicateCalls = 0;
+				InnerJoinBy(
+					FUnboundMapCollectionJoinOp(),
+					[this, &NumPredicateCalls, &Collection0, &Collection2](const FUnboundValueMap* Map0, FUnboundValueMap* Map2)
+					{
+						NumPredicateCalls++;
+						AddErrorIfFalse(CollectionContains(Collection0, Map0), TEXT("Unexpected InnerJoin iteration result"));
+						AddErrorIfFalse(CollectionContains(Collection2, Map2), TEXT("Unexpected InnerJoin iteration result"));
+						AddErrorIfFalse(Map0 != nullptr && Map2 != nullptr && Map0->GetValueType() == Map2->GetValueType(), TEXT("Unexpected InnerJoin iteration result"));
+					},
+					Collection0.CreateConstIterator(),
+					Collection2.CreateIterator());
+				AddErrorIfFalse(NumPredicateCalls == 1, TEXT("Unexpected InnerJoin iteration result"));
+
+				NumPredicateCalls = 0;
+				InnerJoinBy(
+					FUnboundMapCollectionJoinOp(),
+					[this, &NumPredicateCalls, &Collection0, &Collection2](FUnboundValueMap* Map0, const FUnboundValueMap* Map2)
+					{
+						NumPredicateCalls++;
+						AddErrorIfFalse(CollectionContains(Collection0, Map0), TEXT("Unexpected InnerJoin iteration result"));
+						AddErrorIfFalse(CollectionContains(Collection2, Map2), TEXT("Unexpected InnerJoin iteration result"));
+						AddErrorIfFalse(Map0 != nullptr && Map2 != nullptr && Map0->GetValueType() == Map2->GetValueType(), TEXT("Unexpected InnerJoin iteration result"));
+					},
+					Collection0.CreateIterator(),
+					Collection2.CreateConstIterator());
+				AddErrorIfFalse(NumPredicateCalls == 1, TEXT("Unexpected InnerJoin iteration result"));
+			}
+
+			// 3 input variant
+			{
+				int32 NumPredicateCalls = 0;
+				InnerJoinBy(
+					FUnboundMapCollectionJoinOp(),
+					[this, &NumPredicateCalls, &Collection0, &Collection1, &Collection2](FUnboundValueMap* Map0, FUnboundValueMap* Map1, FUnboundValueMap* Map2)
+					{
+						NumPredicateCalls++;
+						AddErrorIfFalse(CollectionContains(Collection0, Map0), TEXT("Unexpected InnerJoin iteration result"));
+						AddErrorIfFalse(CollectionContains(Collection1, Map1), TEXT("Unexpected InnerJoin iteration result"));
+						AddErrorIfFalse(CollectionContains(Collection2, Map2), TEXT("Unexpected InnerJoin iteration result"));
+						AddErrorIfFalse(Map0 != nullptr && Map1 != nullptr && Map0->GetValueType() == Map1->GetValueType(), TEXT("Unexpected InnerJoin iteration result"));
+						AddErrorIfFalse(Map0 != nullptr && Map2 != nullptr && Map0->GetValueType() == Map2->GetValueType(), TEXT("Unexpected InnerJoin iteration result"));
+					},
+					Collection0.CreateIterator(),
+					Collection1.CreateIterator(),
+					Collection2.CreateIterator());
+				AddErrorIfFalse(NumPredicateCalls == 1, TEXT("Unexpected InnerJoin iteration result"));
+
+				NumPredicateCalls = 0;
+				InnerJoinBy(
+					FUnboundMapCollectionJoinOp(),
+					[this, &NumPredicateCalls, &Collection0, &Collection1, &Collection2](FUnboundValueMap* Map0, FUnboundValueMap* Map2, FUnboundValueMap* Map1)
+					{
+						NumPredicateCalls++;
+						AddErrorIfFalse(CollectionContains(Collection0, Map0), TEXT("Unexpected InnerJoin iteration result"));
+						AddErrorIfFalse(CollectionContains(Collection1, Map1), TEXT("Unexpected InnerJoin iteration result"));
+						AddErrorIfFalse(CollectionContains(Collection2, Map2), TEXT("Unexpected InnerJoin iteration result"));
+						AddErrorIfFalse(Map0 != nullptr && Map1 != nullptr && Map0->GetValueType() == Map1->GetValueType(), TEXT("Unexpected InnerJoin iteration result"));
+						AddErrorIfFalse(Map0 != nullptr && Map2 != nullptr && Map0->GetValueType() == Map2->GetValueType(), TEXT("Unexpected InnerJoin iteration result"));
+					},
+					Collection0.CreateIterator(),
+					Collection2.CreateIterator(),
+					Collection1.CreateIterator());
+				AddErrorIfFalse(NumPredicateCalls == 1, TEXT("Unexpected InnerJoin iteration result"));
+
+				NumPredicateCalls = 0;
+				InnerJoinBy(
+					FUnboundMapCollectionJoinOp(),
+					[this, &NumPredicateCalls, &Collection0, &Collection1, &Collection2](FUnboundValueMap* Map1, FUnboundValueMap* Map0, FUnboundValueMap* Map2)
+					{
+						NumPredicateCalls++;
+						AddErrorIfFalse(CollectionContains(Collection0, Map0), TEXT("Unexpected InnerJoin iteration result"));
+						AddErrorIfFalse(CollectionContains(Collection1, Map1), TEXT("Unexpected InnerJoin iteration result"));
+						AddErrorIfFalse(CollectionContains(Collection2, Map2), TEXT("Unexpected InnerJoin iteration result"));
+						AddErrorIfFalse(Map0 != nullptr && Map1 != nullptr && Map0->GetValueType() == Map1->GetValueType(), TEXT("Unexpected InnerJoin iteration result"));
+						AddErrorIfFalse(Map0 != nullptr && Map2 != nullptr && Map0->GetValueType() == Map2->GetValueType(), TEXT("Unexpected InnerJoin iteration result"));
+					},
+					Collection1.CreateIterator(),
+					Collection0.CreateIterator(),
+					Collection2.CreateIterator());
+				AddErrorIfFalse(NumPredicateCalls == 1, TEXT("Unexpected InnerJoin iteration result"));
+
+				NumPredicateCalls = 0;
+				InnerJoinBy(
+					FUnboundMapCollectionJoinOp(),
+					[this, &NumPredicateCalls, &Collection0, &Collection1, &Collection2](FUnboundValueMap* Map1, FUnboundValueMap* Map2, FUnboundValueMap* Map0)
+					{
+						NumPredicateCalls++;
+						AddErrorIfFalse(CollectionContains(Collection0, Map0), TEXT("Unexpected InnerJoin iteration result"));
+						AddErrorIfFalse(CollectionContains(Collection1, Map1), TEXT("Unexpected InnerJoin iteration result"));
+						AddErrorIfFalse(CollectionContains(Collection2, Map2), TEXT("Unexpected InnerJoin iteration result"));
+						AddErrorIfFalse(Map0 != nullptr && Map1 != nullptr && Map0->GetValueType() == Map1->GetValueType(), TEXT("Unexpected InnerJoin iteration result"));
+						AddErrorIfFalse(Map0 != nullptr && Map2 != nullptr && Map0->GetValueType() == Map2->GetValueType(), TEXT("Unexpected InnerJoin iteration result"));
+					},
+					Collection1.CreateIterator(),
+					Collection2.CreateIterator(),
+					Collection0.CreateIterator());
+				AddErrorIfFalse(NumPredicateCalls == 1, TEXT("Unexpected InnerJoin iteration result"));
+
+				NumPredicateCalls = 0;
+				InnerJoinBy(
+					FUnboundMapCollectionJoinOp(),
+					[this, &NumPredicateCalls, &Collection0, &Collection1, &Collection2](FUnboundValueMap* Map2, FUnboundValueMap* Map0, FUnboundValueMap* Map1)
+					{
+						NumPredicateCalls++;
+						AddErrorIfFalse(CollectionContains(Collection0, Map0), TEXT("Unexpected InnerJoin iteration result"));
+						AddErrorIfFalse(CollectionContains(Collection1, Map1), TEXT("Unexpected InnerJoin iteration result"));
+						AddErrorIfFalse(CollectionContains(Collection2, Map2), TEXT("Unexpected InnerJoin iteration result"));
+						AddErrorIfFalse(Map0 != nullptr && Map1 != nullptr && Map0->GetValueType() == Map1->GetValueType(), TEXT("Unexpected InnerJoin iteration result"));
+						AddErrorIfFalse(Map0 != nullptr && Map2 != nullptr && Map0->GetValueType() == Map2->GetValueType(), TEXT("Unexpected InnerJoin iteration result"));
+					},
+					Collection2.CreateIterator(),
+					Collection0.CreateIterator(),
+					Collection1.CreateIterator());
+				AddErrorIfFalse(NumPredicateCalls == 1, TEXT("Unexpected InnerJoin iteration result"));
+
+				NumPredicateCalls = 0;
+				InnerJoinBy(
+					FUnboundMapCollectionJoinOp(),
+					[this, &NumPredicateCalls, &Collection0, &Collection1, &Collection2](FUnboundValueMap* Map2, FUnboundValueMap* Map1, FUnboundValueMap* Map0)
+					{
+						NumPredicateCalls++;
+						AddErrorIfFalse(CollectionContains(Collection0, Map0), TEXT("Unexpected InnerJoin iteration result"));
+						AddErrorIfFalse(CollectionContains(Collection1, Map1), TEXT("Unexpected InnerJoin iteration result"));
+						AddErrorIfFalse(CollectionContains(Collection2, Map2), TEXT("Unexpected InnerJoin iteration result"));
+						AddErrorIfFalse(Map0 != nullptr && Map1 != nullptr && Map0->GetValueType() == Map1->GetValueType(), TEXT("Unexpected InnerJoin iteration result"));
+						AddErrorIfFalse(Map0 != nullptr && Map2 != nullptr && Map0->GetValueType() == Map2->GetValueType(), TEXT("Unexpected InnerJoin iteration result"));
+					},
+					Collection2.CreateIterator(),
+					Collection1.CreateIterator(),
+					Collection0.CreateIterator());
+				AddErrorIfFalse(NumPredicateCalls == 1, TEXT("Unexpected InnerJoin iteration result"));
+			}
+
+			// Join with a transformer list
 			{
 				// Use a float transformer
 				// We set up Collection0 such that it has two mappings with these value types: Float and Bone Transform
@@ -882,7 +904,7 @@
 						AddErrorIfFalse(Map->GetValueType() == FFloatAnimationAttribute::StaticStruct(), TEXT("Unexpected InnerJoin iteration result"));
 					},
 					TransformerList.CreateUnboundValueMapTransformerIterator(),
-					Collection0.CreateMapIterator());
+					Collection0.CreateIterator());
 				AddErrorIfFalse(NumPredicateCalls == 1, TEXT("Unexpected InnerJoin iteration result"));
 			}
 		}
@@ -955,80 +977,67 @@
 			UE_RETURN_ON_ERROR(BindingData.IsValid(), TEXT("Attribute binding data should be valid"));
 		}
 
+		// Bound value map collection
 		{
 			FAttributeNamedSetPtr FullBodySet = BindingData->FindNamedSet(TEXT("FullBody"));
 			UE_RETURN_ON_ERROR(FullBodySet.IsValid(), TEXT("FullBody named set should be valid"));
 
-			FAttributeCollectionHeap Collection0(FullBodySet);
+			FBoundMapCollectionHeap Collection0(FullBodySet);
 
-			AddErrorIfFalse(Collection0.IsEmpty(), TEXT("Attribute collection should be empty"));
+			AddErrorIfFalse(Collection0.IsEmpty(), TEXT("Collection should be empty"));
 			AddErrorIfFalse(Collection0.GetNamedSet() == FullBodySet, TEXT("Unexpected named set"));
 
 			FAttributeMappingKey FloatToFloatKey = FAttributeMappingKey::MakeFromTo<FFloatAnimationAttribute, FFloatAnimationAttribute>();
-			TBoundValueMap<FFloatAnimationAttribute>* FloatToFloatMap0 = Collection0.AddSetMap<FFloatAnimationAttribute>(FloatToFloatKey);
+			TBoundValueMap<FFloatAnimationAttribute>* FloatToFloatMap0 = Collection0.Add<FFloatAnimationAttribute>(FloatToFloatKey);
 			UE_RETURN_ON_ERROR(FloatToFloatMap0 != nullptr, TEXT("Float->Float map should be valid"));
 
 			FAttributeMappingKey FloatToIntegerKey = FAttributeMappingKey::MakeFromTo<FFloatAnimationAttribute, FIntegerAnimationAttribute>();
-			TBoundValueMap<FIntegerAnimationAttribute>* FloatToIntegerMap0 = Collection0.AddSetMap<FIntegerAnimationAttribute>(FloatToIntegerKey);
+			TBoundValueMap<FIntegerAnimationAttribute>* FloatToIntegerMap0 = Collection0.Add<FIntegerAnimationAttribute>(FloatToIntegerKey);
 			UE_RETURN_ON_ERROR(FloatToIntegerMap0 != nullptr, TEXT("Float->Integer map should be valid"));
 
 			FAttributeMappingKey FloatToStringKey = FAttributeMappingKey::MakeFromTo<FFloatAnimationAttribute, FStringAnimationAttribute>();
-			TBoundValueMap<FStringAnimationAttribute>* FloatToStringMap0 = Collection0.AddSetMap<FStringAnimationAttribute>(FloatToStringKey);
+			TBoundValueMap<FStringAnimationAttribute>* FloatToStringMap0 = Collection0.Add<FStringAnimationAttribute>(FloatToStringKey);
 			UE_RETURN_ON_ERROR(FloatToStringMap0 != nullptr, TEXT("Float->String map should be valid"));
 
 			FAttributeMappingKey TransformToTransformKey = FAttributeMappingKey::MakeFromTo<FBoneTransformAnimationAttribute, FBoneTransformAnimationAttribute>();
-			TBoundValueMap<FBoneTransformAnimationAttribute>* TransformToTransformMap0 = Collection0.AddSetMap<FBoneTransformAnimationAttribute>(TransformToTransformKey);
+			TBoundValueMap<FBoneTransformAnimationAttribute>* TransformToTransformMap0 = Collection0.Add<FBoneTransformAnimationAttribute>(TransformToTransformKey);
 			UE_RETURN_ON_ERROR(TransformToTransformMap0 != nullptr, TEXT("Transform->Transform map should be valid"));
 
 			FAttributeMappingKey TransformToFloatKey = FAttributeMappingKey::MakeFromTo<FBoneTransformAnimationAttribute, FFloatAnimationAttribute>();
-			TBoundValueMap<FFloatAnimationAttribute>* TransformToFloatMap0 = Collection0.AddSetMap<FFloatAnimationAttribute>(TransformToFloatKey);
+			TBoundValueMap<FFloatAnimationAttribute>* TransformToFloatMap0 = Collection0.Add<FFloatAnimationAttribute>(TransformToFloatKey);
 			UE_RETURN_ON_ERROR(TransformToFloatMap0 != nullptr, TEXT("Transform->Bone map should be valid"));
 
-			TUnboundValueMap<FFloatAnimationAttribute>* FloatMap0 = Collection0.AddMap<FFloatAnimationAttribute>();
-			UE_RETURN_ON_ERROR(FloatMap0 != nullptr, TEXT("Float map should be valid"));
+			AddErrorIfFalse(!Collection0.IsEmpty(), TEXT("Collection should not be empty"));
+			AddErrorIfFalse(Collection0.Num() == 5, TEXT("Unexpected map count"));
 
-			TUnboundValueMap<FBoneTransformAnimationAttribute>* BoneMap0 = Collection0.AddMap<FBoneTransformAnimationAttribute>();
-			UE_RETURN_ON_ERROR(BoneMap0 != nullptr, TEXT("Bone map should be valid"));
+			FBoundMapCollectionHeap Collection1(FullBodySet);
 
-			AddErrorIfFalse(!Collection0.IsEmpty(), TEXT("Attribute collection should not be empty"));
-			AddErrorIfFalse(Collection0.NumSetMaps() == 5, TEXT("Unexpected set map count"));
-			AddErrorIfFalse(Collection0.NumMaps() == 2, TEXT("Unexpected map count"));
-
-			FAttributeCollectionHeap Collection1(FullBodySet);
-
-			AddErrorIfFalse(Collection1.IsEmpty(), TEXT("Attribute collection should be empty"));
+			AddErrorIfFalse(Collection1.IsEmpty(), TEXT("Collection should be empty"));
 			AddErrorIfFalse(Collection1.GetNamedSet() == FullBodySet, TEXT("Unexpected named set"));
 
-			TBoundValueMap<FFloatAnimationAttribute>* FloatToFloatMap1 = Collection1.AddSetMap<FFloatAnimationAttribute>(FloatToFloatKey);
+			TBoundValueMap<FFloatAnimationAttribute>* FloatToFloatMap1 = Collection1.Add<FFloatAnimationAttribute>(FloatToFloatKey);
 			UE_RETURN_ON_ERROR(FloatToFloatMap1 != nullptr, TEXT("Float->Float map should be valid"));
 
-			TBoundValueMap<FIntegerAnimationAttribute>* FloatToIntegerMap1 = Collection1.AddSetMap<FIntegerAnimationAttribute>(FloatToIntegerKey);
+			TBoundValueMap<FIntegerAnimationAttribute>* FloatToIntegerMap1 = Collection1.Add<FIntegerAnimationAttribute>(FloatToIntegerKey);
 			UE_RETURN_ON_ERROR(FloatToIntegerMap1 != nullptr, TEXT("Float->Integer map should be valid"));
 
-			TBoundValueMap<FStringAnimationAttribute>* FloatToStringMap1 = Collection1.AddSetMap<FStringAnimationAttribute>(FloatToStringKey);
+			TBoundValueMap<FStringAnimationAttribute>* FloatToStringMap1 = Collection1.Add<FStringAnimationAttribute>(FloatToStringKey);
 			UE_RETURN_ON_ERROR(FloatToStringMap1 != nullptr, TEXT("Float->String map should be valid"));
 
-			TBoundValueMap<FBoneTransformAnimationAttribute>* TransformToTransformMap1 = Collection1.AddSetMap<FBoneTransformAnimationAttribute>(TransformToTransformKey);
+			TBoundValueMap<FBoneTransformAnimationAttribute>* TransformToTransformMap1 = Collection1.Add<FBoneTransformAnimationAttribute>(TransformToTransformKey);
 			UE_RETURN_ON_ERROR(TransformToTransformMap1 != nullptr, TEXT("Transform->Transform map should be valid"));
 
-			TBoundValueMap<FFloatAnimationAttribute>* TransformToFloatMap1 = Collection1.AddSetMap<FFloatAnimationAttribute>(TransformToFloatKey);
+			TBoundValueMap<FFloatAnimationAttribute>* TransformToFloatMap1 = Collection1.Add<FFloatAnimationAttribute>(TransformToFloatKey);
 			UE_RETURN_ON_ERROR(TransformToFloatMap1 != nullptr, TEXT("Transform->Bone map should be valid"));
 
-			TUnboundValueMap<FFloatAnimationAttribute>* FloatMap1 = Collection1.AddMap<FFloatAnimationAttribute>();
-			UE_RETURN_ON_ERROR(FloatMap1 != nullptr, TEXT("Float map should be valid"));
+			AddErrorIfFalse(!Collection1.IsEmpty(), TEXT("Collection should not be empty"));
+			AddErrorIfFalse(Collection1.Num() == 5, TEXT("Unexpected map count"));
 
-			TUnboundValueMap<FBoneTransformAnimationAttribute>* BoneMap1 = Collection1.AddMap<FBoneTransformAnimationAttribute>();
-			UE_RETURN_ON_ERROR(BoneMap1 != nullptr, TEXT("Bone map should be valid"));
-
-			AddErrorIfFalse(!Collection1.IsEmpty(), TEXT("Attribute collection should not be empty"));
-			AddErrorIfFalse(Collection1.NumSetMaps() == 5, TEXT("Unexpected set map count"));
-			AddErrorIfFalse(Collection1.NumMaps() == 2, TEXT("Unexpected map count"));
-
-			// Both inputs perfectly match (set maps)
+			// Both inputs perfectly match
 			{
 				int32 NumPredicateCalls = 0;
 				OuterJoinBy(
-					FAttributeCollectionSetMapJoinOp(),
+					FBoundMapCollectionJoinOp(),
 					[this, &NumPredicateCalls, &Collection0, &Collection1](FBoundValueMap* Map0, FBoundValueMap* Map1)
 					{
 						NumPredicateCalls++;
@@ -1037,13 +1046,13 @@
 						AddErrorIfFalse(Map0 != nullptr && Map1 != nullptr && Map0->GetAttributeType() == Map1->GetAttributeType(), TEXT("Unexpected OuterJoin iteration result"));
 						AddErrorIfFalse(Map0 != nullptr && Map1 != nullptr && Map0->GetValueType() == Map1->GetValueType(), TEXT("Unexpected OuterJoin iteration result"));
 					},
-					Collection0.CreateSetMapIterator(),
-					Collection1.CreateSetMapIterator());
+					Collection0.CreateIterator(),
+					Collection1.CreateIterator());
 				AddErrorIfFalse(NumPredicateCalls == 5, TEXT("Unexpected OuterJoin iteration result"));
 
 				NumPredicateCalls = 0;
 				OuterJoinBy(
-					FAttributeCollectionSetMapJoinOp(),
+					FBoundMapCollectionJoinOp(),
 					[this, &NumPredicateCalls, &Collection0, &Collection1](const FBoundValueMap* Map0, const FBoundValueMap* Map1)
 					{
 						NumPredicateCalls++;
@@ -1052,13 +1061,13 @@
 						AddErrorIfFalse(Map0 != nullptr && Map1 != nullptr && Map0->GetAttributeType() == Map1->GetAttributeType(), TEXT("Unexpected OuterJoin iteration result"));
 						AddErrorIfFalse(Map0 != nullptr && Map1 != nullptr && Map0->GetValueType() == Map1->GetValueType(), TEXT("Unexpected OuterJoin iteration result"));
 					},
-					Collection0.CreateConstSetMapIterator(),
-					Collection1.CreateConstSetMapIterator());
+					Collection0.CreateConstIterator(),
+					Collection1.CreateConstIterator());
 				AddErrorIfFalse(NumPredicateCalls == 5, TEXT("Unexpected OuterJoin iteration result"));
 
 				NumPredicateCalls = 0;
 				OuterJoinBy(
-					FAttributeCollectionSetMapJoinOp(),
+					FBoundMapCollectionJoinOp(),
 					[this, &NumPredicateCalls, &Collection0, &Collection1](const FBoundValueMap* Map0, FBoundValueMap* Map1)
 					{
 						NumPredicateCalls++;
@@ -1067,13 +1076,13 @@
 						AddErrorIfFalse(Map0 != nullptr && Map1 != nullptr && Map0->GetAttributeType() == Map1->GetAttributeType(), TEXT("Unexpected OuterJoin iteration result"));
 						AddErrorIfFalse(Map0 != nullptr && Map1 != nullptr && Map0->GetValueType() == Map1->GetValueType(), TEXT("Unexpected OuterJoin iteration result"));
 					},
-					Collection0.CreateConstSetMapIterator(),
-					Collection1.CreateSetMapIterator());
+					Collection0.CreateConstIterator(),
+					Collection1.CreateIterator());
 				AddErrorIfFalse(NumPredicateCalls == 5, TEXT("Unexpected OuterJoin iteration result"));
 
 				NumPredicateCalls = 0;
 				OuterJoinBy(
-					FAttributeCollectionSetMapJoinOp(),
+					FBoundMapCollectionJoinOp(),
 					[this, &NumPredicateCalls, &Collection0, &Collection1](FBoundValueMap* Map0, const FBoundValueMap* Map1)
 					{
 						NumPredicateCalls++;
@@ -1082,94 +1091,31 @@
 						AddErrorIfFalse(Map0 != nullptr && Map1 != nullptr && Map0->GetAttributeType() == Map1->GetAttributeType(), TEXT("Unexpected OuterJoin iteration result"));
 						AddErrorIfFalse(Map0 != nullptr && Map1 != nullptr && Map0->GetValueType() == Map1->GetValueType(), TEXT("Unexpected OuterJoin iteration result"));
 					},
-					Collection0.CreateSetMapIterator(),
-					Collection1.CreateConstSetMapIterator());
+					Collection0.CreateIterator(),
+					Collection1.CreateConstIterator());
 				AddErrorIfFalse(NumPredicateCalls == 5, TEXT("Unexpected OuterJoin iteration result"));
 			}
 
-			// Both inputs perfectly match (name/value maps)
-			{
-				int32 NumPredicateCalls = 0;
-				OuterJoinBy(
-					FAttributeCollectionMapJoinOp(),
-					[this, &NumPredicateCalls, &Collection0, &Collection1](FUnboundValueMap* Map0, FUnboundValueMap* Map1)
-					{
-						NumPredicateCalls++;
-						AddErrorIfFalse(CollectionContains(Collection0, Map0), TEXT("Unexpected OuterJoin iteration result"));
-						AddErrorIfFalse(CollectionContains(Collection1, Map1), TEXT("Unexpected OuterJoin iteration result"));
-						AddErrorIfFalse(Map0 != nullptr && Map1 != nullptr && Map0->GetValueType() == Map1->GetValueType(), TEXT("Unexpected OuterJoin iteration result"));
-					},
-					Collection0.CreateMapIterator(),
-					Collection1.CreateMapIterator());
-				AddErrorIfFalse(NumPredicateCalls == 2, TEXT("Unexpected OuterJoin iteration result"));
+			FBoundMapCollectionHeap Collection2(FullBodySet);
 
-				NumPredicateCalls = 0;
-				OuterJoinBy(
-					FAttributeCollectionMapJoinOp(),
-					[this, &NumPredicateCalls, &Collection0, &Collection1](const FUnboundValueMap* Map0, const FUnboundValueMap* Map1)
-					{
-						NumPredicateCalls++;
-						AddErrorIfFalse(CollectionContains(Collection0, Map0), TEXT("Unexpected OuterJoin iteration result"));
-						AddErrorIfFalse(CollectionContains(Collection1, Map1), TEXT("Unexpected OuterJoin iteration result"));
-						AddErrorIfFalse(Map0 != nullptr && Map1 != nullptr && Map0->GetValueType() == Map1->GetValueType(), TEXT("Unexpected OuterJoin iteration result"));
-					},
-					Collection0.CreateConstMapIterator(),
-					Collection1.CreateConstMapIterator());
-				AddErrorIfFalse(NumPredicateCalls == 2, TEXT("Unexpected OuterJoin iteration result"));
-
-				NumPredicateCalls = 0;
-				OuterJoinBy(
-					FAttributeCollectionMapJoinOp(),
-					[this, &NumPredicateCalls, &Collection0, &Collection1](const FUnboundValueMap* Map0, FUnboundValueMap* Map1)
-					{
-						NumPredicateCalls++;
-						AddErrorIfFalse(CollectionContains(Collection0, Map0), TEXT("Unexpected OuterJoin iteration result"));
-						AddErrorIfFalse(CollectionContains(Collection1, Map1), TEXT("Unexpected OuterJoin iteration result"));
-						AddErrorIfFalse(Map0 != nullptr && Map1 != nullptr && Map0->GetValueType() == Map1->GetValueType(), TEXT("Unexpected OuterJoin iteration result"));
-					},
-					Collection0.CreateConstMapIterator(),
-					Collection1.CreateMapIterator());
-				AddErrorIfFalse(NumPredicateCalls == 2, TEXT("Unexpected OuterJoin iteration result"));
-
-				NumPredicateCalls = 0;
-				OuterJoinBy(
-					FAttributeCollectionMapJoinOp(),
-					[this, &NumPredicateCalls, &Collection0, &Collection1](FUnboundValueMap* Map0, const FUnboundValueMap* Map1)
-					{
-						NumPredicateCalls++;
-						AddErrorIfFalse(CollectionContains(Collection0, Map0), TEXT("Unexpected OuterJoin iteration result"));
-						AddErrorIfFalse(CollectionContains(Collection1, Map1), TEXT("Unexpected OuterJoin iteration result"));
-						AddErrorIfFalse(Map0 != nullptr && Map1 != nullptr && Map0->GetValueType() == Map1->GetValueType(), TEXT("Unexpected OuterJoin iteration result"));
-					},
-					Collection0.CreateMapIterator(),
-					Collection1.CreateConstMapIterator());
-				AddErrorIfFalse(NumPredicateCalls == 2, TEXT("Unexpected OuterJoin iteration result"));
-			}
-
-			FAttributeCollectionHeap Collection2(FullBodySet);
-
-			AddErrorIfFalse(Collection2.IsEmpty(), TEXT("Attribute collection should be empty"));
+			AddErrorIfFalse(Collection2.IsEmpty(), TEXT("Collection should be empty"));
 			AddErrorIfFalse(Collection2.GetNamedSet() == FullBodySet, TEXT("Unexpected named set"));
 
-			TBoundValueMap<FFloatAnimationAttribute>* FloatToFloatMap2 = Collection2.AddSetMap<FFloatAnimationAttribute>(FloatToFloatKey);
+			TBoundValueMap<FFloatAnimationAttribute>* FloatToFloatMap2 = Collection2.Add<FFloatAnimationAttribute>(FloatToFloatKey);
 			UE_RETURN_ON_ERROR(FloatToFloatMap2 != nullptr, TEXT("Float->Float map should be valid"));
 
-			TBoundValueMap<FBoneTransformAnimationAttribute>* TransformToTransformMap2 = Collection2.AddSetMap<FBoneTransformAnimationAttribute>(TransformToTransformKey);
+			TBoundValueMap<FBoneTransformAnimationAttribute>* TransformToTransformMap2 = Collection2.Add<FBoneTransformAnimationAttribute>(TransformToTransformKey);
 			UE_RETURN_ON_ERROR(TransformToTransformMap2 != nullptr, TEXT("Transform->Transform map should be valid"));
 
-			TUnboundValueMap<FFloatAnimationAttribute>* FloatMap2 = Collection2.AddMap<FFloatAnimationAttribute>();
-			UE_RETURN_ON_ERROR(FloatMap2 != nullptr, TEXT("Float map should be valid"));
+			AddErrorIfFalse(!Collection2.IsEmpty(), TEXT("Collection should not be empty"));
+			AddErrorIfFalse(Collection2.Num() == 2, TEXT("Unexpected map count"));
 
-			AddErrorIfFalse(!Collection2.IsEmpty(), TEXT("Attribute collection should not be empty"));
-			AddErrorIfFalse(Collection2.NumSetMaps() == 2, TEXT("Unexpected set map count"));
-			AddErrorIfFalse(Collection2.NumMaps() == 1, TEXT("Unexpected map count"));
-
-			// First input is shorter and will be padded with null entries (set maps)
+			// First input is shorter and will be padded with null entries
 			{
 				int32 NumPredicateCalls = 0;
 				int32 NumNullEntries = 0;
 				OuterJoinBy(
-					FAttributeCollectionSetMapJoinOp(),
+					FBoundMapCollectionJoinOp(),
 					[this, &NumPredicateCalls, &NumNullEntries, &Collection0, &Collection2](FBoundValueMap* Map2, FBoundValueMap* Map0)
 					{
 						NumPredicateCalls++;
@@ -1187,15 +1133,15 @@
 							AddErrorIfFalse(Map0 != nullptr, TEXT("Unexpected OuterJoin iteration result"));
 						}
 					},
-					Collection2.CreateSetMapIterator(),
-					Collection0.CreateSetMapIterator());
+					Collection2.CreateIterator(),
+					Collection0.CreateIterator());
 				AddErrorIfFalse(NumPredicateCalls == 5, TEXT("Unexpected OuterJoin iteration result"));
 				AddErrorIfFalse(NumNullEntries == 3, TEXT("Unexpected OuterJoin iteration result"));
 
 				NumPredicateCalls = 0;
 				NumNullEntries = 0;
 				OuterJoinBy(
-					FAttributeCollectionSetMapJoinOp(),
+					FBoundMapCollectionJoinOp(),
 					[this, &NumPredicateCalls, &NumNullEntries, &Collection0, &Collection2](const FBoundValueMap* Map2, const FBoundValueMap* Map0)
 					{
 						NumPredicateCalls++;
@@ -1213,15 +1159,15 @@
 							AddErrorIfFalse(Map0 != nullptr, TEXT("Unexpected OuterJoin iteration result"));
 						}
 					},
-					Collection2.CreateConstSetMapIterator(),
-					Collection0.CreateConstSetMapIterator());
+					Collection2.CreateConstIterator(),
+					Collection0.CreateConstIterator());
 				AddErrorIfFalse(NumPredicateCalls == 5, TEXT("Unexpected OuterJoin iteration result"));
 				AddErrorIfFalse(NumNullEntries == 3, TEXT("Unexpected OuterJoin iteration result"));
 
 				NumPredicateCalls = 0;
 				NumNullEntries = 0;
 				OuterJoinBy(
-					FAttributeCollectionSetMapJoinOp(),
+					FBoundMapCollectionJoinOp(),
 					[this, &NumPredicateCalls, &NumNullEntries, &Collection0, &Collection2](const FBoundValueMap* Map2, FBoundValueMap* Map0)
 					{
 						NumPredicateCalls++;
@@ -1239,15 +1185,15 @@
 							AddErrorIfFalse(Map0 != nullptr, TEXT("Unexpected OuterJoin iteration result"));
 						}
 					},
-					Collection2.CreateConstSetMapIterator(),
-					Collection0.CreateSetMapIterator());
+					Collection2.CreateConstIterator(),
+					Collection0.CreateIterator());
 				AddErrorIfFalse(NumPredicateCalls == 5, TEXT("Unexpected OuterJoin iteration result"));
 				AddErrorIfFalse(NumNullEntries == 3, TEXT("Unexpected OuterJoin iteration result"));
 
 				NumPredicateCalls = 0;
 				NumNullEntries = 0;
 				OuterJoinBy(
-					FAttributeCollectionSetMapJoinOp(),
+					FBoundMapCollectionJoinOp(),
 					[this, &NumPredicateCalls, &NumNullEntries, &Collection0, &Collection2](FBoundValueMap* Map2, const FBoundValueMap* Map0)
 					{
 						NumPredicateCalls++;
@@ -1265,19 +1211,19 @@
 							AddErrorIfFalse(Map0 != nullptr, TEXT("Unexpected OuterJoin iteration result"));
 						}
 					},
-					Collection2.CreateSetMapIterator(),
-					Collection0.CreateConstSetMapIterator());
+					Collection2.CreateIterator(),
+					Collection0.CreateConstIterator());
 				AddErrorIfFalse(NumPredicateCalls == 5, TEXT("Unexpected OuterJoin iteration result"));
 				AddErrorIfFalse(NumNullEntries == 3, TEXT("Unexpected OuterJoin iteration result"));
 			}
 
-			// First input is shorter and will be padded with null entries (name/value maps)
+			// Second input is shorter and will be padded with null entries
 			{
 				int32 NumPredicateCalls = 0;
 				int32 NumNullEntries = 0;
 				OuterJoinBy(
-					FAttributeCollectionMapJoinOp(),
-					[this, &NumPredicateCalls, &NumNullEntries, &Collection0, &Collection2](FUnboundValueMap* Map2, FUnboundValueMap* Map0)
+					FBoundMapCollectionJoinOp(),
+					[this, &NumPredicateCalls, &NumNullEntries, &Collection0, &Collection2](FBoundValueMap* Map0, FBoundValueMap* Map2)
 					{
 						NumPredicateCalls++;
 						AddErrorIfFalse(CollectionContains(Collection0, Map0), TEXT("Unexpected OuterJoin iteration result"));
@@ -1285,6 +1231,7 @@
 						if (Map2 != nullptr)
 						{
 							AddErrorIfFalse(CollectionContains(Collection2, Map2), TEXT("Unexpected OuterJoin iteration result"));
+							AddErrorIfFalse(Map0 != nullptr && Map0->GetAttributeType() == Map2->GetAttributeType(), TEXT("Unexpected OuterJoin iteration result"));
 							AddErrorIfFalse(Map0 != nullptr && Map0->GetValueType() == Map2->GetValueType(), TEXT("Unexpected OuterJoin iteration result"));
 						}
 						else
@@ -1293,16 +1240,16 @@
 							AddErrorIfFalse(Map0 != nullptr, TEXT("Unexpected OuterJoin iteration result"));
 						}
 					},
-					Collection2.CreateMapIterator(),
-					Collection0.CreateMapIterator());
-				AddErrorIfFalse(NumPredicateCalls == 2, TEXT("Unexpected OuterJoin iteration result"));
-				AddErrorIfFalse(NumNullEntries == 1, TEXT("Unexpected OuterJoin iteration result"));
+					Collection0.CreateIterator(),
+					Collection2.CreateIterator());
+				AddErrorIfFalse(NumPredicateCalls == 5, TEXT("Unexpected OuterJoin iteration result"));
+				AddErrorIfFalse(NumNullEntries == 3, TEXT("Unexpected OuterJoin iteration result"));
 
 				NumPredicateCalls = 0;
 				NumNullEntries = 0;
 				OuterJoinBy(
-					FAttributeCollectionMapJoinOp(),
-					[this, &NumPredicateCalls, &NumNullEntries, &Collection0, &Collection2](const FUnboundValueMap* Map2, const FUnboundValueMap* Map0)
+					FBoundMapCollectionJoinOp(),
+					[this, &NumPredicateCalls, &NumNullEntries, &Collection0, &Collection2](const FBoundValueMap* Map0, const FBoundValueMap* Map2)
 					{
 						NumPredicateCalls++;
 						AddErrorIfFalse(CollectionContains(Collection0, Map0), TEXT("Unexpected OuterJoin iteration result"));
@@ -1310,6 +1257,7 @@
 						if (Map2 != nullptr)
 						{
 							AddErrorIfFalse(CollectionContains(Collection2, Map2), TEXT("Unexpected OuterJoin iteration result"));
+							AddErrorIfFalse(Map0 != nullptr && Map0->GetAttributeType() == Map2->GetAttributeType(), TEXT("Unexpected OuterJoin iteration result"));
 							AddErrorIfFalse(Map0 != nullptr && Map0->GetValueType() == Map2->GetValueType(), TEXT("Unexpected OuterJoin iteration result"));
 						}
 						else
@@ -1318,16 +1266,16 @@
 							AddErrorIfFalse(Map0 != nullptr, TEXT("Unexpected OuterJoin iteration result"));
 						}
 					},
-					Collection2.CreateConstMapIterator(),
-					Collection0.CreateConstMapIterator());
-				AddErrorIfFalse(NumPredicateCalls == 2, TEXT("Unexpected OuterJoin iteration result"));
-				AddErrorIfFalse(NumNullEntries == 1, TEXT("Unexpected OuterJoin iteration result"));
+					Collection0.CreateConstIterator(),
+					Collection2.CreateConstIterator());
+				AddErrorIfFalse(NumPredicateCalls == 5, TEXT("Unexpected OuterJoin iteration result"));
+				AddErrorIfFalse(NumNullEntries == 3, TEXT("Unexpected OuterJoin iteration result"));
 
 				NumPredicateCalls = 0;
 				NumNullEntries = 0;
 				OuterJoinBy(
-					FAttributeCollectionMapJoinOp(),
-					[this, &NumPredicateCalls, &NumNullEntries, &Collection0, &Collection2](const FUnboundValueMap* Map2, FUnboundValueMap* Map0)
+					FBoundMapCollectionJoinOp(),
+					[this, &NumPredicateCalls, &NumNullEntries, &Collection0, &Collection2](const FBoundValueMap* Map0, FBoundValueMap* Map2)
 					{
 						NumPredicateCalls++;
 						AddErrorIfFalse(CollectionContains(Collection0, Map0), TEXT("Unexpected OuterJoin iteration result"));
@@ -1335,6 +1283,7 @@
 						if (Map2 != nullptr)
 						{
 							AddErrorIfFalse(CollectionContains(Collection2, Map2), TEXT("Unexpected OuterJoin iteration result"));
+							AddErrorIfFalse(Map0 != nullptr && Map0->GetAttributeType() == Map2->GetAttributeType(), TEXT("Unexpected OuterJoin iteration result"));
 							AddErrorIfFalse(Map0 != nullptr && Map0->GetValueType() == Map2->GetValueType(), TEXT("Unexpected OuterJoin iteration result"));
 						}
 						else
@@ -1343,16 +1292,16 @@
 							AddErrorIfFalse(Map0 != nullptr, TEXT("Unexpected OuterJoin iteration result"));
 						}
 					},
-					Collection2.CreateConstMapIterator(),
-					Collection0.CreateMapIterator());
-				AddErrorIfFalse(NumPredicateCalls == 2, TEXT("Unexpected OuterJoin iteration result"));
-				AddErrorIfFalse(NumNullEntries == 1, TEXT("Unexpected OuterJoin iteration result"));
+					Collection0.CreateConstIterator(),
+					Collection2.CreateIterator());
+				AddErrorIfFalse(NumPredicateCalls == 5, TEXT("Unexpected OuterJoin iteration result"));
+				AddErrorIfFalse(NumNullEntries == 3, TEXT("Unexpected OuterJoin iteration result"));
 
 				NumPredicateCalls = 0;
 				NumNullEntries = 0;
 				OuterJoinBy(
-					FAttributeCollectionMapJoinOp(),
-					[this, &NumPredicateCalls, &NumNullEntries, &Collection0, &Collection2](FUnboundValueMap* Map2, const FUnboundValueMap* Map0)
+					FBoundMapCollectionJoinOp(),
+					[this, &NumPredicateCalls, &NumNullEntries, &Collection0, &Collection2](FBoundValueMap* Map0, const FBoundValueMap* Map2)
 					{
 						NumPredicateCalls++;
 						AddErrorIfFalse(CollectionContains(Collection0, Map0), TEXT("Unexpected OuterJoin iteration result"));
@@ -1360,6 +1309,7 @@
 						if (Map2 != nullptr)
 						{
 							AddErrorIfFalse(CollectionContains(Collection2, Map2), TEXT("Unexpected OuterJoin iteration result"));
+							AddErrorIfFalse(Map0 != nullptr && Map0->GetAttributeType() == Map2->GetAttributeType(), TEXT("Unexpected OuterJoin iteration result"));
 							AddErrorIfFalse(Map0 != nullptr && Map0->GetValueType() == Map2->GetValueType(), TEXT("Unexpected OuterJoin iteration result"));
 						}
 						else
@@ -1368,19 +1318,200 @@
 							AddErrorIfFalse(Map0 != nullptr, TEXT("Unexpected OuterJoin iteration result"));
 						}
 					},
-					Collection2.CreateMapIterator(),
-					Collection0.CreateConstMapIterator());
-				AddErrorIfFalse(NumPredicateCalls == 2, TEXT("Unexpected OuterJoin iteration result"));
-				AddErrorIfFalse(NumNullEntries == 1, TEXT("Unexpected OuterJoin iteration result"));
+					Collection0.CreateIterator(),
+					Collection2.CreateConstIterator());
+				AddErrorIfFalse(NumPredicateCalls == 5, TEXT("Unexpected OuterJoin iteration result"));
+				AddErrorIfFalse(NumNullEntries == 3, TEXT("Unexpected OuterJoin iteration result"));
 			}
 
-			// Second input is shorter and will be padded with null entries (set maps)
+			// Join with a transformer list
+			{
+				// Use a float transformer
+				// We set up Collection0 such that it has two mappings with float values: Transform -> Float and Float -> Float
+				// InnerJoin should execute the predicate twice, once for each value type that matches the transformer's value type
+
+				// Dummy float sanitize transformer
+				struct FFloatAttributeTransformer_SanitizeTest
+				{
+					using TransformerType = Transformers::FSanitize;
+					using ValueType = FFloatAnimationAttribute;
+
+					static void TransformBoundValueMap(FBoundValueMap* SetMap) {}	// No-op
+				};
+
+				// Dummy vector sanitize transformer
+				struct FVectorAttributeTransformer_SanitizeTest
+				{
+					using TransformerType = Transformers::FSanitize;
+					using ValueType = FVectorAnimationAttribute;
+
+					static void TransformBoundValueMap(FBoundValueMap* SetMap) {}	// No-op
+				};
+
+#ifdef _MSC_VER
+#pragma warning(push)
+#pragma warning(disable : 4191) // unsafe conversion from 'function pointer type A' to 'function pointer type B'
+#endif
+
+				FValueTransformerList TransformerList(Transformers::FSanitize::TransformerName);
+				TransformerList.AddBoundValueMapTransformer(FFloatAnimationAttribute::StaticStruct(), reinterpret_cast<FRawTransformerFunc>(&FFloatAttributeTransformer_SanitizeTest::TransformBoundValueMap));
+				TransformerList.AddBoundValueMapTransformer(FVectorAnimationAttribute::StaticStruct(), reinterpret_cast<FRawTransformerFunc>(&FVectorAttributeTransformer_SanitizeTest::TransformBoundValueMap));
+
+#ifdef _MSC_VER
+#pragma warning(pop)
+#endif
+
+				int32 NumPredicateCalls = 0;
+				int32 NumNullCollectionEntries = 0;
+				int32 NumNullTransformEntries = 0;
+				OuterJoinBy(
+					FValueTransformerListWithBoundMapCollectionJoinOp(),
+					[this, &NumPredicateCalls, &NumNullCollectionEntries, &NumNullTransformEntries, &Collection0](FRawTransformerFunc Transformer, FBoundValueMap* SetMap)
+					{
+						NumPredicateCalls++;
+
+						if (SetMap != nullptr)
+						{
+							AddErrorIfFalse(CollectionContains(Collection0, SetMap), TEXT("Unexpected OuterJoin iteration result"));
+
+							if (Transformer != nullptr)
+							{
+								AddErrorIfFalse(SetMap->GetValueType() == FFloatAnimationAttribute::StaticStruct(), TEXT("Unexpected OuterJoin iteration result"));
+							}
+						}
+						else
+						{
+							NumNullCollectionEntries++;
+						}
+
+						if (Transformer == nullptr)
+						{
+							NumNullTransformEntries++;
+						}
+					},
+					TransformerList.CreateBoundValueMapTransformerIterator(),
+					Collection0.CreateIterator());
+
+				// [Transformer Value Type, Collection Set Map Value Type]
+				// ---
+				// [Float, Float->Float]
+				// [Float, Bone->Float]
+				// [Float, nullptr]
+				// [nullptr, Float->Int]
+				// [nullptr, Float->String]
+				// [nullptr, Bone->Bone]
+				// [Vector, nullptr]
+
+				AddErrorIfFalse(NumPredicateCalls == 7, TEXT("Unexpected OuterJoin iteration result"));
+				AddErrorIfFalse(NumNullCollectionEntries == 2, TEXT("Unexpected OuterJoin iteration result"));
+				AddErrorIfFalse(NumNullTransformEntries == 3, TEXT("Unexpected OuterJoin iteration result"));
+			}
+		}
+
+		// Unbound value map collection
+		{
+			FUnboundMapCollectionHeap Collection0;
+
+			AddErrorIfFalse(Collection0.IsEmpty(), TEXT("Collection should be empty"));
+
+			TUnboundValueMap<FFloatAnimationAttribute>* FloatMap0 = Collection0.Add<FFloatAnimationAttribute>();
+			UE_RETURN_ON_ERROR(FloatMap0 != nullptr, TEXT("Float map should be valid"));
+
+			TUnboundValueMap<FBoneTransformAnimationAttribute>* BoneMap0 = Collection0.Add<FBoneTransformAnimationAttribute>();
+			UE_RETURN_ON_ERROR(BoneMap0 != nullptr, TEXT("Bone map should be valid"));
+
+			AddErrorIfFalse(!Collection0.IsEmpty(), TEXT("Collection should not be empty"));
+			AddErrorIfFalse(Collection0.Num() == 2, TEXT("Unexpected map count"));
+
+			FUnboundMapCollectionHeap Collection1;
+
+			AddErrorIfFalse(Collection1.IsEmpty(), TEXT("Collection should be empty"));
+
+			TUnboundValueMap<FFloatAnimationAttribute>* FloatMap1 = Collection1.Add<FFloatAnimationAttribute>();
+			UE_RETURN_ON_ERROR(FloatMap1 != nullptr, TEXT("Float map should be valid"));
+
+			TUnboundValueMap<FBoneTransformAnimationAttribute>* BoneMap1 = Collection1.Add<FBoneTransformAnimationAttribute>();
+			UE_RETURN_ON_ERROR(BoneMap1 != nullptr, TEXT("Bone map should be valid"));
+
+			AddErrorIfFalse(!Collection1.IsEmpty(), TEXT("Collection should not be empty"));
+			AddErrorIfFalse(Collection1.Num() == 2, TEXT("Unexpected map count"));
+
+			// Both inputs perfectly match
 			{
 				int32 NumPredicateCalls = 0;
+				OuterJoinBy(
+					FUnboundMapCollectionJoinOp(),
+					[this, &NumPredicateCalls, &Collection0, &Collection1](FUnboundValueMap* Map0, FUnboundValueMap* Map1)
+					{
+						NumPredicateCalls++;
+						AddErrorIfFalse(CollectionContains(Collection0, Map0), TEXT("Unexpected OuterJoin iteration result"));
+						AddErrorIfFalse(CollectionContains(Collection1, Map1), TEXT("Unexpected OuterJoin iteration result"));
+						AddErrorIfFalse(Map0 != nullptr && Map1 != nullptr && Map0->GetValueType() == Map1->GetValueType(), TEXT("Unexpected OuterJoin iteration result"));
+					},
+					Collection0.CreateIterator(),
+					Collection1.CreateIterator());
+				AddErrorIfFalse(NumPredicateCalls == 2, TEXT("Unexpected OuterJoin iteration result"));
+
+				NumPredicateCalls = 0;
+				OuterJoinBy(
+					FUnboundMapCollectionJoinOp(),
+					[this, &NumPredicateCalls, &Collection0, &Collection1](const FUnboundValueMap* Map0, const FUnboundValueMap* Map1)
+					{
+						NumPredicateCalls++;
+						AddErrorIfFalse(CollectionContains(Collection0, Map0), TEXT("Unexpected OuterJoin iteration result"));
+						AddErrorIfFalse(CollectionContains(Collection1, Map1), TEXT("Unexpected OuterJoin iteration result"));
+						AddErrorIfFalse(Map0 != nullptr && Map1 != nullptr && Map0->GetValueType() == Map1->GetValueType(), TEXT("Unexpected OuterJoin iteration result"));
+					},
+					Collection0.CreateConstIterator(),
+					Collection1.CreateConstIterator());
+				AddErrorIfFalse(NumPredicateCalls == 2, TEXT("Unexpected OuterJoin iteration result"));
+
+				NumPredicateCalls = 0;
+				OuterJoinBy(
+					FUnboundMapCollectionJoinOp(),
+					[this, &NumPredicateCalls, &Collection0, &Collection1](const FUnboundValueMap* Map0, FUnboundValueMap* Map1)
+					{
+						NumPredicateCalls++;
+						AddErrorIfFalse(CollectionContains(Collection0, Map0), TEXT("Unexpected OuterJoin iteration result"));
+						AddErrorIfFalse(CollectionContains(Collection1, Map1), TEXT("Unexpected OuterJoin iteration result"));
+						AddErrorIfFalse(Map0 != nullptr && Map1 != nullptr && Map0->GetValueType() == Map1->GetValueType(), TEXT("Unexpected OuterJoin iteration result"));
+					},
+					Collection0.CreateConstIterator(),
+					Collection1.CreateIterator());
+				AddErrorIfFalse(NumPredicateCalls == 2, TEXT("Unexpected OuterJoin iteration result"));
+
+				NumPredicateCalls = 0;
+				OuterJoinBy(
+					FUnboundMapCollectionJoinOp(),
+					[this, &NumPredicateCalls, &Collection0, &Collection1](FUnboundValueMap* Map0, const FUnboundValueMap* Map1)
+					{
+						NumPredicateCalls++;
+						AddErrorIfFalse(CollectionContains(Collection0, Map0), TEXT("Unexpected OuterJoin iteration result"));
+						AddErrorIfFalse(CollectionContains(Collection1, Map1), TEXT("Unexpected OuterJoin iteration result"));
+						AddErrorIfFalse(Map0 != nullptr && Map1 != nullptr && Map0->GetValueType() == Map1->GetValueType(), TEXT("Unexpected OuterJoin iteration result"));
+					},
+					Collection0.CreateIterator(),
+					Collection1.CreateConstIterator());
+				AddErrorIfFalse(NumPredicateCalls == 2, TEXT("Unexpected OuterJoin iteration result"));
+			}
+
+			FUnboundMapCollectionHeap Collection2;
+
+			AddErrorIfFalse(Collection2.IsEmpty(), TEXT("Collection should be empty"));
+
+			TUnboundValueMap<FFloatAnimationAttribute>* FloatMap2 = Collection2.Add<FFloatAnimationAttribute>();
+			UE_RETURN_ON_ERROR(FloatMap2 != nullptr, TEXT("Float map should be valid"));
+
+			AddErrorIfFalse(!Collection2.IsEmpty(), TEXT("Collection should not be empty"));
+			AddErrorIfFalse(Collection2.Num() == 1, TEXT("Unexpected map count"));
+
+			// First input is shorter and will be padded with null entries
+			{
+				int32 NumPredicateCalls = 0;
 				int32 NumNullEntries = 0;
 				OuterJoinBy(
-					FAttributeCollectionSetMapJoinOp(),
-					[this, &NumPredicateCalls, &NumNullEntries, &Collection0, &Collection2](FBoundValueMap* Map0, FBoundValueMap* Map2)
+					FUnboundMapCollectionJoinOp(),
+					[this, &NumPredicateCalls, &NumNullEntries, &Collection0, &Collection2](FUnboundValueMap* Map2, FUnboundValueMap* Map0)
 					{
 						NumPredicateCalls++;
 						AddErrorIfFalse(CollectionContains(Collection0, Map0), TEXT("Unexpected OuterJoin iteration result"));
@@ -1388,7 +1519,6 @@
 						if (Map2 != nullptr)
 						{
 							AddErrorIfFalse(CollectionContains(Collection2, Map2), TEXT("Unexpected OuterJoin iteration result"));
-							AddErrorIfFalse(Map0 != nullptr && Map0->GetAttributeType() == Map2->GetAttributeType(), TEXT("Unexpected OuterJoin iteration result"));
 							AddErrorIfFalse(Map0 != nullptr && Map0->GetValueType() == Map2->GetValueType(), TEXT("Unexpected OuterJoin iteration result"));
 						}
 						else
@@ -1397,16 +1527,16 @@
 							AddErrorIfFalse(Map0 != nullptr, TEXT("Unexpected OuterJoin iteration result"));
 						}
 					},
-					Collection0.CreateSetMapIterator(),
-					Collection2.CreateSetMapIterator());
-				AddErrorIfFalse(NumPredicateCalls == 5, TEXT("Unexpected OuterJoin iteration result"));
-				AddErrorIfFalse(NumNullEntries == 3, TEXT("Unexpected OuterJoin iteration result"));
+					Collection2.CreateIterator(),
+					Collection0.CreateIterator());
+				AddErrorIfFalse(NumPredicateCalls == 2, TEXT("Unexpected OuterJoin iteration result"));
+				AddErrorIfFalse(NumNullEntries == 1, TEXT("Unexpected OuterJoin iteration result"));
 
 				NumPredicateCalls = 0;
 				NumNullEntries = 0;
 				OuterJoinBy(
-					FAttributeCollectionSetMapJoinOp(),
-					[this, &NumPredicateCalls, &NumNullEntries, &Collection0, &Collection2](const FBoundValueMap* Map0, const FBoundValueMap* Map2)
+					FUnboundMapCollectionJoinOp(),
+					[this, &NumPredicateCalls, &NumNullEntries, &Collection0, &Collection2](const FUnboundValueMap* Map2, const FUnboundValueMap* Map0)
 					{
 						NumPredicateCalls++;
 						AddErrorIfFalse(CollectionContains(Collection0, Map0), TEXT("Unexpected OuterJoin iteration result"));
@@ -1414,7 +1544,6 @@
 						if (Map2 != nullptr)
 						{
 							AddErrorIfFalse(CollectionContains(Collection2, Map2), TEXT("Unexpected OuterJoin iteration result"));
-							AddErrorIfFalse(Map0 != nullptr && Map0->GetAttributeType() == Map2->GetAttributeType(), TEXT("Unexpected OuterJoin iteration result"));
 							AddErrorIfFalse(Map0 != nullptr && Map0->GetValueType() == Map2->GetValueType(), TEXT("Unexpected OuterJoin iteration result"));
 						}
 						else
@@ -1423,16 +1552,16 @@
 							AddErrorIfFalse(Map0 != nullptr, TEXT("Unexpected OuterJoin iteration result"));
 						}
 					},
-					Collection0.CreateConstSetMapIterator(),
-					Collection2.CreateConstSetMapIterator());
-				AddErrorIfFalse(NumPredicateCalls == 5, TEXT("Unexpected OuterJoin iteration result"));
-				AddErrorIfFalse(NumNullEntries == 3, TEXT("Unexpected OuterJoin iteration result"));
+					Collection2.CreateConstIterator(),
+					Collection0.CreateConstIterator());
+				AddErrorIfFalse(NumPredicateCalls == 2, TEXT("Unexpected OuterJoin iteration result"));
+				AddErrorIfFalse(NumNullEntries == 1, TEXT("Unexpected OuterJoin iteration result"));
 
 				NumPredicateCalls = 0;
 				NumNullEntries = 0;
 				OuterJoinBy(
-					FAttributeCollectionSetMapJoinOp(),
-					[this, &NumPredicateCalls, &NumNullEntries, &Collection0, &Collection2](const FBoundValueMap* Map0, FBoundValueMap* Map2)
+					FUnboundMapCollectionJoinOp(),
+					[this, &NumPredicateCalls, &NumNullEntries, &Collection0, &Collection2](const FUnboundValueMap* Map2, FUnboundValueMap* Map0)
 					{
 						NumPredicateCalls++;
 						AddErrorIfFalse(CollectionContains(Collection0, Map0), TEXT("Unexpected OuterJoin iteration result"));
@@ -1440,7 +1569,6 @@
 						if (Map2 != nullptr)
 						{
 							AddErrorIfFalse(CollectionContains(Collection2, Map2), TEXT("Unexpected OuterJoin iteration result"));
-							AddErrorIfFalse(Map0 != nullptr && Map0->GetAttributeType() == Map2->GetAttributeType(), TEXT("Unexpected OuterJoin iteration result"));
 							AddErrorIfFalse(Map0 != nullptr && Map0->GetValueType() == Map2->GetValueType(), TEXT("Unexpected OuterJoin iteration result"));
 						}
 						else
@@ -1449,16 +1577,16 @@
 							AddErrorIfFalse(Map0 != nullptr, TEXT("Unexpected OuterJoin iteration result"));
 						}
 					},
-					Collection0.CreateConstSetMapIterator(),
-					Collection2.CreateSetMapIterator());
-				AddErrorIfFalse(NumPredicateCalls == 5, TEXT("Unexpected OuterJoin iteration result"));
-				AddErrorIfFalse(NumNullEntries == 3, TEXT("Unexpected OuterJoin iteration result"));
+					Collection2.CreateConstIterator(),
+					Collection0.CreateIterator());
+				AddErrorIfFalse(NumPredicateCalls == 2, TEXT("Unexpected OuterJoin iteration result"));
+				AddErrorIfFalse(NumNullEntries == 1, TEXT("Unexpected OuterJoin iteration result"));
 
 				NumPredicateCalls = 0;
 				NumNullEntries = 0;
 				OuterJoinBy(
-					FAttributeCollectionSetMapJoinOp(),
-					[this, &NumPredicateCalls, &NumNullEntries, &Collection0, &Collection2](FBoundValueMap* Map0, const FBoundValueMap* Map2)
+					FUnboundMapCollectionJoinOp(),
+					[this, &NumPredicateCalls, &NumNullEntries, &Collection0, &Collection2](FUnboundValueMap* Map2, const FUnboundValueMap* Map0)
 					{
 						NumPredicateCalls++;
 						AddErrorIfFalse(CollectionContains(Collection0, Map0), TEXT("Unexpected OuterJoin iteration result"));
@@ -1466,7 +1594,6 @@
 						if (Map2 != nullptr)
 						{
 							AddErrorIfFalse(CollectionContains(Collection2, Map2), TEXT("Unexpected OuterJoin iteration result"));
-							AddErrorIfFalse(Map0 != nullptr && Map0->GetAttributeType() == Map2->GetAttributeType(), TEXT("Unexpected OuterJoin iteration result"));
 							AddErrorIfFalse(Map0 != nullptr && Map0->GetValueType() == Map2->GetValueType(), TEXT("Unexpected OuterJoin iteration result"));
 						}
 						else
@@ -1475,18 +1602,18 @@
 							AddErrorIfFalse(Map0 != nullptr, TEXT("Unexpected OuterJoin iteration result"));
 						}
 					},
-					Collection0.CreateSetMapIterator(),
-					Collection2.CreateConstSetMapIterator());
-				AddErrorIfFalse(NumPredicateCalls == 5, TEXT("Unexpected OuterJoin iteration result"));
-				AddErrorIfFalse(NumNullEntries == 3, TEXT("Unexpected OuterJoin iteration result"));
+					Collection2.CreateIterator(),
+					Collection0.CreateConstIterator());
+				AddErrorIfFalse(NumPredicateCalls == 2, TEXT("Unexpected OuterJoin iteration result"));
+				AddErrorIfFalse(NumNullEntries == 1, TEXT("Unexpected OuterJoin iteration result"));
 			}
 
-			// Second input is shorter and will be padded with null entries (name/value maps)
+			// Second input is shorter and will be padded with null entries
 			{
 				int32 NumPredicateCalls = 0;
 				int32 NumNullEntries = 0;
 				OuterJoinBy(
-					FAttributeCollectionMapJoinOp(),
+					FUnboundMapCollectionJoinOp(),
 					[this, &NumPredicateCalls, &NumNullEntries, &Collection0, &Collection2](FUnboundValueMap* Map0, FUnboundValueMap* Map2)
 					{
 						NumPredicateCalls++;
@@ -1503,15 +1630,15 @@
 							AddErrorIfFalse(Map0 != nullptr, TEXT("Unexpected OuterJoin iteration result"));
 						}
 					},
-					Collection0.CreateMapIterator(),
-					Collection2.CreateMapIterator());
+					Collection0.CreateIterator(),
+					Collection2.CreateIterator());
 				AddErrorIfFalse(NumPredicateCalls == 2, TEXT("Unexpected OuterJoin iteration result"));
 				AddErrorIfFalse(NumNullEntries == 1, TEXT("Unexpected OuterJoin iteration result"));
 
 				NumPredicateCalls = 0;
 				NumNullEntries = 0;
 				OuterJoinBy(
-					FAttributeCollectionMapJoinOp(),
+					FUnboundMapCollectionJoinOp(),
 					[this, &NumPredicateCalls, &NumNullEntries, &Collection0, &Collection2](const FUnboundValueMap* Map0, const FUnboundValueMap* Map2)
 					{
 						NumPredicateCalls++;
@@ -1528,15 +1655,15 @@
 							AddErrorIfFalse(Map0 != nullptr, TEXT("Unexpected OuterJoin iteration result"));
 						}
 					},
-					Collection0.CreateConstMapIterator(),
-					Collection2.CreateConstMapIterator());
+					Collection0.CreateConstIterator(),
+					Collection2.CreateConstIterator());
 				AddErrorIfFalse(NumPredicateCalls == 2, TEXT("Unexpected OuterJoin iteration result"));
 				AddErrorIfFalse(NumNullEntries == 1, TEXT("Unexpected OuterJoin iteration result"));
 
 				NumPredicateCalls = 0;
 				NumNullEntries = 0;
 				OuterJoinBy(
-					FAttributeCollectionMapJoinOp(),
+					FUnboundMapCollectionJoinOp(),
 					[this, &NumPredicateCalls, &NumNullEntries, &Collection0, &Collection2](const FUnboundValueMap* Map0, FUnboundValueMap* Map2)
 					{
 						NumPredicateCalls++;
@@ -1553,15 +1680,15 @@
 							AddErrorIfFalse(Map0 != nullptr, TEXT("Unexpected OuterJoin iteration result"));
 						}
 					},
-					Collection0.CreateConstMapIterator(),
-					Collection2.CreateMapIterator());
+					Collection0.CreateConstIterator(),
+					Collection2.CreateIterator());
 				AddErrorIfFalse(NumPredicateCalls == 2, TEXT("Unexpected OuterJoin iteration result"));
 				AddErrorIfFalse(NumNullEntries == 1, TEXT("Unexpected OuterJoin iteration result"));
 
 				NumPredicateCalls = 0;
 				NumNullEntries = 0;
 				OuterJoinBy(
-					FAttributeCollectionMapJoinOp(),
+					FUnboundMapCollectionJoinOp(),
 					[this, &NumPredicateCalls, &NumNullEntries, &Collection0, &Collection2](FUnboundValueMap* Map0, const FUnboundValueMap* Map2)
 					{
 						NumPredicateCalls++;
@@ -1578,96 +1705,13 @@
 							AddErrorIfFalse(Map0 != nullptr, TEXT("Unexpected OuterJoin iteration result"));
 						}
 					},
-					Collection0.CreateMapIterator(),
-					Collection2.CreateConstMapIterator());
+					Collection0.CreateIterator(),
+					Collection2.CreateConstIterator());
 				AddErrorIfFalse(NumPredicateCalls == 2, TEXT("Unexpected OuterJoin iteration result"));
 				AddErrorIfFalse(NumNullEntries == 1, TEXT("Unexpected OuterJoin iteration result"));
 			}
 
-			// Join with a transformer list (set maps)
+			// Join with a transformer list
 			{
 				// Use a float transformer
-				// We set up Collection0 such that it has two mappings with float values: Transform -> Float and Float -> Float
-				// InnerJoin should execute the predicate twice, once for each value type that matches the transformer's value type
-
-				// Dummy float sanitize transformer
-				struct FFloatAttributeTransformer_SanitizeTest
-				{
-					using TransformerType = Transformers::FSanitize;
-					using ValueType = FFloatAnimationAttribute;
-
-					static void TransformBoundValueMap(FBoundValueMap* SetMap) {}	// No-op
-				};
-
-				// Dummy vector sanitize transformer
-				struct FVectorAttributeTransformer_SanitizeTest
-				{
-					using TransformerType = Transformers::FSanitize;
-					using ValueType = FVectorAnimationAttribute;
-
-					static void TransformBoundValueMap(FBoundValueMap* SetMap) {}	// No-op
-				};
-
-#ifdef _MSC_VER
-	#pragma warning(push)
-	#pragma warning(disable : 4191) // unsafe conversion from 'function pointer type A' to 'function pointer type B'
-#endif
-
-				FValueTransformerList TransformerList(Transformers::FSanitize::TransformerName);
-				TransformerList.AddBoundValueMapTransformer(FFloatAnimationAttribute::StaticStruct(), reinterpret_cast<FRawTransformerFunc>(&FFloatAttributeTransformer_SanitizeTest::TransformBoundValueMap));
-				TransformerList.AddBoundValueMapTransformer(FVectorAnimationAttribute::StaticStruct(), reinterpret_cast<FRawTransformerFunc>(&FVectorAttributeTransformer_SanitizeTest::TransformBoundValueMap));
-
-#ifdef _MSC_VER
-	#pragma warning(pop)
-#endif
-
-				int32 NumPredicateCalls = 0;
-				int32 NumNullCollectionEntries = 0;
-				int32 NumNullTransformEntries = 0;
-				OuterJoinBy(
-					FValueTransformerListWithBoundMapCollectionJoinOp(),
-					[this, &NumPredicateCalls, &NumNullCollectionEntries, &NumNullTransformEntries, &Collection0](FRawTransformerFunc Transformer, FBoundValueMap* SetMap)
-					{
-						NumPredicateCalls++;
-
-						if (SetMap != nullptr)
-						{
-							AddErrorIfFalse(CollectionContains(Collection0, SetMap), TEXT("Unexpected OuterJoin iteration result"));
-
-							if (Transformer != nullptr)
-							{
-								AddErrorIfFalse(SetMap->GetValueType() == FFloatAnimationAttribute::StaticStruct(), TEXT("Unexpected OuterJoin iteration result"));
-							}
-						}
-						else
-						{
-							NumNullCollectionEntries++;
-						}
-
-						if (Transformer == nullptr)
-						{
-							NumNullTransformEntries++;
-						}
-					},
-					TransformerList.CreateBoundValueMapTransformerIterator(),
-					Collection0.CreateSetMapIterator());
-
-				// [Transformer Value Type, Collection Set Map Value Type]
-				// ---
-				// [Float, Float->Float]
-				// [Float, Bone->Float]
-				// [Float, nullptr]
-				// [nullptr, Float->Int]
-				// [nullptr, Float->String]
-				// [nullptr, Bone->Bone]
-				// [Vector, nullptr]
-				
-				AddErrorIfFalse(NumPredicateCalls == 7, TEXT("Unexpected OuterJoin iteration result"));
-				AddErrorIfFalse(NumNullCollectionEntries == 2, TEXT("Unexpected OuterJoin iteration result"));
-				AddErrorIfFalse(NumNullTransformEntries == 3, TEXT("Unexpected OuterJoin iteration result"));
-			}
-
-			// Join with a transformer list (name/value maps)
-			{
-				// Use a float transformer
 				// We set up Collection0 such that it has two mappings with these value types: Float and Bone Transform
@@ -1733,7 +1777,7 @@
 						}
 					},
 					TransformerList.CreateUnboundValueMapTransformerIterator(),
-					Collection0.CreateMapIterator());
+					Collection0.CreateIterator());
 
 				// [Transformer Value Type, Collection Map Value Type]
 				// ---
