Rename UAnimNextModule_EditorData to UUAFSystem_EditorData
Rename UAnimNextModuleFactory to UUAFSystemFactory

#Synced-CL 49591051

diff --git a/Engine/Plugins/Experimental/UAF/UAF/Config/DefaultUAF.ini b/Engine/Plugins/Experimental/UAF/UAF/Config/DefaultUAF.ini
--- a/Engine/Plugins/Experimental/UAF/UAF/Config/DefaultUAF.ini
+++ b/Engine/Plugins/Experimental/UAF/UAF/Config/DefaultUAF.ini
@@ -5,7 +5,7 @@
 
 ; Parameter block/graph merge
 +ClassRedirects=(OldName="/Script/AnimNext.AnimNextParameterBlock", NewName="/Script/UAF.UAFSystem")
-+ClassRedirects=(OldName="/Script/AnimNextUncookedOnly.AnimNextParameterBlock_EditorData", NewName="/Script/UAFUncookedOnly.AnimNextModule_EditorData")
++ClassRedirects=(OldName="/Script/AnimNextUncookedOnly.AnimNextParameterBlock_EditorData", NewName="/Script/UAFUncookedOnly.UAFSystem_EditorData")
 +ClassRedirects=(OldName="/Script/AnimNextUncookedOnly.AnimNextParameterBlock_EdGraph", NewName="/Script/UAFUncookedOnly.AnimNextEdGraph")
 +ClassRedirects=(OldName="/Script/AnimNextUncookedOnly.AnimNextParameterBlock_EdGraphNode", NewName="/Script/UAFUncookedOnly.AnimNextEdGraphNode")
 +ClassRedirects=(OldName="/Script/AnimNextUncookedOnly.AnimNextParameterBlock_EdGraphSchema", NewName="/Script/UAFUncookedOnly.AnimNextEdGraphSchema")
@@ -25,7 +25,7 @@
 ; AnimNextGraph -> AnimNextModule
 +ClassRedirects=(OldName="/Script/AnimNext.AnimNextGraph", NewName="/Script/UAF.UAFSystem")
 +ClassRedirects=(OldName="/Script/AnimNextUncookedOnly.AnimNextGraph_Controller", NewName="/Script/UAFUncookedOnly.AnimNextControllerBase")
-+ClassRedirects=(OldName="/Script/AnimNextUncookedOnly.AnimNextGraph_EditorData", NewName="/Script/UAFUncookedOnly.AnimNextModule_EditorData")
++ClassRedirects=(OldName="/Script/AnimNextUncookedOnly.AnimNextGraph_EditorData", NewName="/Script/UAFUncookedOnly.UAFSystem_EditorData")
 +ClassRedirects=(OldName="/Script/AnimNextUncookedOnly.AnimNextGraph_AnimationGraphSchema", NewName="/Script/UAFAnimGraphUncookedOnly.AnimNextAnimationGraphSchema")
 +ClassRedirects=(OldName="/Script/AnimNextUncookedOnly.AnimNextGraph_EdGraph", NewName="/Script/UAFUncookedOnly.AnimNextEdGraph")
 +ClassRedirects=(OldName="/Script/AnimNextUncookedOnly.AnimNextGraph_EdGraphNode", NewName="/Script/UAFUncookedOnly.AnimNextEdGraphNode")
@@ -277,7 +277,7 @@
 +ClassRedirects=(OldName="/Script/AnimNextUncookedOnly.AnimNextSharedVariablesEntry", NewName="/Script/UAFUncookedOnly.AnimNextSharedVariablesEntry")
 +ClassRedirects=(OldName="/Script/AnimNextUncookedOnly.AnimNextVariableEntry", NewName="/Script/UAFUncookedOnly.AnimNextVariableEntry")
 +ClassRedirects=(OldName="/Script/AnimNextUncookedOnly.AnimNextEventGraphSchema", NewName="/Script/UAFUncookedOnly.AnimNextEventGraphSchema")
-+ClassRedirects=(OldName="/Script/AnimNextUncookedOnly.AnimNextModule_EditorData", NewName="/Script/UAFUncookedOnly.AnimNextModule_EditorData")
++ClassRedirects=(OldName="/Script/AnimNextUncookedOnly.AnimNextModule_EditorData", NewName="/Script/UAFUncookedOnly.UAFSystem_EditorData")
 +ClassRedirects=(OldName="/Script/AnimNextUncookedOnly.AnimNextSharedVariables_EditorData", NewName="/Script/UAFUncookedOnly.AnimNextSharedVariables_EditorData")
 +ClassRedirects=(OldName="/Script/AnimNextUncookedOnly.K2Node_AnimNextComponentSetVariable", NewName="/Script/UAFUncookedOnly.K2Node_AnimNextComponentSetVariable")
 
@@ -304,4 +304,6 @@
 +FunctionRedirects=(OldName="AnimNextComponent.BlueprintGetModuleHandle",NewName="AnimNextComponent.GetSystemReference")
 
 ; AnimNext -> UAF
-+ClassRedirects=(OldName="/Script/UAF.AnimNextModule",NewName="/Script/UAF.UAFSystem")
\ No newline at end of file
++ClassRedirects=(OldName="/Script/UAF.AnimNextModule",NewName="/Script/UAF.UAFSystem")
++ClassRedirects=(OldName="/Script/UAFUncookedOnly.AnimNextModule_EditorData",NewName="/Script/UAFUncookedOnly.UAFSystem_EditorData")
++ClassRedirects=(OldName="/Script/UAFUncookedOnly.AnimNextModuleFactory",NewName="/Script/UAFUncookedOnly.UAFSystemFactory")
\ No newline at end of file
diff --git a/Engine/Plugins/Experimental/UAF/UAF/Source/UAF/Internal/Module/AnimNextModule.h b/Engine/Plugins/Experimental/UAF/UAF/Source/UAF/Internal/Module/AnimNextModule.h
--- a/Engine/Plugins/Experimental/UAF/UAF/Source/UAF/Internal/Module/AnimNextModule.h
+++ b/Engine/Plugins/Experimental/UAF/UAF/Source/UAF/Internal/Module/AnimNextModule.h
@@ -58,8 +58,8 @@ public:
 	UE_API virtual void PostLoad() override;
 
 protected:
-	friend class UAnimNextModuleFactory;
-	friend class UAnimNextModule_EditorData;
+	friend class UUAFSystemFactory;
+	friend class UUAFSystem_EditorData;
 	friend struct UE::UAF::UncookedOnly::FUtils;
 	friend class UE::UAF::Editor::FModuleEditor;
 	friend struct UE::UAF::FTestUtils;
diff --git a/Engine/Plugins/Experimental/UAF/UAF/Source/UAFEditor/Internal/Common/AnimNextAssetItemDetails.cpp b/Engine/Plugins/Experimental/UAF/UAF/Source/UAFEditor/Internal/Common/AnimNextAssetItemDetails.cpp
--- a/Engine/Plugins/Experimental/UAF/UAF/Source/UAFEditor/Internal/Common/AnimNextAssetItemDetails.cpp
+++ b/Engine/Plugins/Experimental/UAF/UAF/Source/UAFEditor/Internal/Common/AnimNextAssetItemDetails.cpp
@@ -186,7 +186,7 @@ void FAnimNextAssetItemDetails::RegisterToolMenuExtensions()
 								));
 						};
 
-						if(EditorData->IsA<UAnimNextModule_EditorData>())
+						if(EditorData->IsA<UUAFSystem_EditorData>())
 						{
 							for(TObjectIterator<UScriptStruct> It; It; ++It)
 							{
diff --git a/Engine/Plugins/Experimental/UAF/UAF/Source/UAFEditor/Internal/Module/AnimNextModuleFactory.h b/Engine/Plugins/Experimental/UAF/UAF/Source/UAFEditor/Internal/Module/AnimNextModuleFactory.h
--- a/Engine/Plugins/Experimental/UAF/UAF/Source/UAFEditor/Internal/Module/AnimNextModuleFactory.h
+++ b/Engine/Plugins/Experimental/UAF/UAF/Source/UAFEditor/Internal/Module/AnimNextModuleFactory.h
@@ -7,11 +7,11 @@
 #include "AnimNextModuleFactory.generated.h"
 
 UCLASS(MinimalAPI)
-class UAnimNextModuleFactory : public UFactory
+class UUAFSystemFactory : public UFactory
 {
 	GENERATED_BODY()
 
-	UAnimNextModuleFactory();
+	UUAFSystemFactory();
 	
 	virtual bool ConfigureProperties() override;
 	virtual UObject* FactoryCreateNew(UClass* Class, UObject* InParent, FName Name, EObjectFlags Flags, UObject* Context, FFeedbackContext* Warn, FName CallingContext) override;
diff --git a/Engine/Plugins/Experimental/UAF/UAF/Source/UAFEditor/Private/AnimNextEditorModule.cpp b/Engine/Plugins/Experimental/UAF/UAF/Source/UAFEditor/Private/AnimNextEditorModule.cpp
--- a/Engine/Plugins/Experimental/UAF/UAF/Source/UAFEditor/Private/AnimNextEditorModule.cpp
+++ b/Engine/Plugins/Experimental/UAF/UAF/Source/UAFEditor/Private/AnimNextEditorModule.cpp
@@ -913,7 +913,7 @@ void FAnimNextEditorModule::RegisterWorkspaceDocumentTypes(Workspace::IWorkspace
 				UObject* InObject = Document.GetObject();
 				if(UUAFSystem* Module = Cast<UUAFSystem>(InObject))
 				{
-				    UAnimNextModule_EditorData* EditorData = UncookedOnly::FUtils::GetEditorData<UAnimNextModule_EditorData>(Module);
+				    UUAFSystem_EditorData* EditorData = UncookedOnly::FUtils::GetEditorData<UUAFSystem_EditorData>(Module);
 
 					const TInstancedStruct<FWorkspaceOutlinerItemData>& ExportData = Document.Export.GetData();
 					if (ExportData.GetScriptStruct() == FAnimNextGraphOutlinerData::StaticStruct())
diff --git a/Engine/Plugins/Experimental/UAF/UAF/Source/UAFEditor/Private/Module/AnimNextModuleFactory.cpp b/Engine/Plugins/Experimental/UAF/UAF/Source/UAFEditor/Private/Module/AnimNextModuleFactory.cpp
--- a/Engine/Plugins/Experimental/UAF/UAF/Source/UAFEditor/Private/Module/AnimNextModuleFactory.cpp
+++ b/Engine/Plugins/Experimental/UAF/UAF/Source/UAFEditor/Private/Module/AnimNextModuleFactory.cpp
@@ -11,19 +11,19 @@
 
 #include UE_INLINE_GENERATED_CPP_BY_NAME(AnimNextModuleFactory)
 
-UAnimNextModuleFactory::UAnimNextModuleFactory()
+UUAFSystemFactory::UUAFSystemFactory()
 {
 	bCreateNew = true;
 	bEditAfterNew = true;
 	SupportedClass = UUAFSystem::StaticClass();
 }
 
-bool UAnimNextModuleFactory::ConfigureProperties()
+bool UUAFSystemFactory::ConfigureProperties()
 {
 	return true;
 }
 
-UObject* UAnimNextModuleFactory::FactoryCreateNew(UClass* Class, UObject* InParent, FName Name, EObjectFlags Flags, UObject* Context, FFeedbackContext* Warn, FName CallingContext)
+UObject* UUAFSystemFactory::FactoryCreateNew(UClass* Class, UObject* InParent, FName Name, EObjectFlags Flags, UObject* Context, FFeedbackContext* Warn, FName CallingContext)
 {
 	EObjectFlags FlagsToUse = Flags | RF_Public | RF_Standalone | RF_Transactional | RF_LoadCompleted;
 	if(InParent == GetTransientPackage())
@@ -34,12 +34,12 @@ UObject* UAnimNextModuleFactory::FactoryCreateNew(UClass* Class, UObject* InPare
 	UUAFSystem* NewModule = NewObject<UUAFSystem>(InParent, Class, Name, FlagsToUse);
 
 	// Create internal editor data
-	UAnimNextModule_EditorData* EditorData = NewObject<UAnimNextModule_EditorData>(NewModule, TEXT("EditorData"), RF_Transactional);
+	UUAFSystem_EditorData* EditorData = NewObject<UUAFSystem_EditorData>(NewModule, TEXT("EditorData"), RF_Transactional);
 	NewModule->EditorData = EditorData;
 	EditorData->bUsesExternalPackages = false;
 	EditorData->Initialize(/*bRecompileVM*/false);
 
-	UAnimNextEventGraphEntry* GraphEntry = EditorData->AddEventGraph(UAnimNextModule_EditorData::DefaultEventGraphName, FRigUnit_AnimNextInitializeEvent::StaticStruct());
+	UAnimNextEventGraphEntry* GraphEntry = EditorData->AddEventGraph(UUAFSystem_EditorData::DefaultEventGraphName, FRigUnit_AnimNextInitializeEvent::StaticStruct());
 
 	URigVMController* Controller = EditorData->GetController(GraphEntry->Graph);
 	Controller->AddUnitNode(FRigUnit_AnimNextPrePhysicsEvent::StaticStruct(), FRigVMStruct::ExecuteName, FVector2D(0.0f, 200.0f), FString(), false);
diff --git a/Engine/Plugins/Experimental/UAF/UAF/Source/UAFUncookedOnly/Internal/AnimNextEdGraph.h b/Engine/Plugins/Experimental/UAF/UAF/Source/UAFUncookedOnly/Internal/AnimNextEdGraph.h
--- a/Engine/Plugins/Experimental/UAF/UAF/Source/UAFUncookedOnly/Internal/AnimNextEdGraph.h
+++ b/Engine/Plugins/Experimental/UAF/UAF/Source/UAFUncookedOnly/Internal/AnimNextEdGraph.h
@@ -21,7 +21,7 @@ class UAnimNextEdGraph : public URigVMEdGraph
 {
 	GENERATED_BODY()
 
-	friend class UAnimNextModule_EditorData;
+	friend class UUAFSystem_EditorData;
 	friend class UAnimNextRigVMAssetEditorData;
 
 	// UObject interface
diff --git a/Engine/Plugins/Experimental/UAF/UAF/Source/UAFUncookedOnly/Internal/Entries/AnimNextEventGraphEntry.h b/Engine/Plugins/Experimental/UAF/UAF/Source/UAFUncookedOnly/Internal/Entries/AnimNextEventGraphEntry.h
--- a/Engine/Plugins/Experimental/UAF/UAF/Source/UAFUncookedOnly/Internal/Entries/AnimNextEventGraphEntry.h
+++ b/Engine/Plugins/Experimental/UAF/UAF/Source/UAFUncookedOnly/Internal/Entries/AnimNextEventGraphEntry.h
@@ -10,7 +10,7 @@
 #include "AnimNextEventGraphEntry.generated.h"
 
 class UAnimNextRigVMAssetEditorData;
-class UAnimNextModule_EditorData;
+class UUAFSystem_EditorData;
 class UAnimNextEdGraph;
 
 namespace UE::UAF::Tests
@@ -34,8 +34,8 @@ class UAnimNextEventGraphEntry : public UAnimNextRigVMAssetEntry, public IAnimNe
 	GENERATED_BODY()
 
 	friend class UAnimNextRigVMAssetEditorData;
-	friend class UAnimNextModule_EditorData;
-	friend class UAnimNextModuleFactory;
+	friend class UUAFSystem_EditorData;
+	friend class UUAFSystemFactory;
 	friend class UE::UAF::Tests::FEditor_Variables;
 	friend class UE::UAF::Tests::FVariables_UOLBindings;
 	friend class UE::UAF::Tests::FEditor_AnimGraph_Variables;
diff --git a/Engine/Plugins/Experimental/UAF/UAF/Source/UAFUncookedOnly/Internal/Entries/AnimNextVariableEntry.h b/Engine/Plugins/Experimental/UAF/UAF/Source/UAFUncookedOnly/Internal/Entries/AnimNextVariableEntry.h
--- a/Engine/Plugins/Experimental/UAF/UAF/Source/UAFUncookedOnly/Internal/Entries/AnimNextVariableEntry.h
+++ b/Engine/Plugins/Experimental/UAF/UAF/Source/UAFUncookedOnly/Internal/Entries/AnimNextVariableEntry.h
@@ -13,7 +13,7 @@
 #include "AnimNextVariableEntry.generated.h"
 
 class UAnimNextRigVMAssetEditorData;
-class UAnimNextModule_EditorData;
+class UUAFSystem_EditorData;
 class UAnimNextSharedVariables_EditorData;
 class UAnimNextVariableBinding;
 class UAnimNextSharedVariablesEntry;
@@ -60,7 +60,7 @@ class UAnimNextVariableEntry : public UAnimNextRigVMAssetEntry, public IAnimNext
 	GENERATED_BODY()
 
 	friend class UAnimNextRigVMAssetEditorData;
-	friend class UAnimNextModule_EditorData;
+	friend class UUAFSystem_EditorData;
 	friend class UAnimNextSharedVariables_EditorData;
 	friend class UAnimNextSharedVariablesEntry;
 	friend class UAnimNextStateTreeTreeEditorData;
diff --git a/Engine/Plugins/Experimental/UAF/UAF/Source/UAFUncookedOnly/Internal/Module/AnimNextModule_EditorData.h b/Engine/Plugins/Experimental/UAF/UAF/Source/UAFUncookedOnly/Internal/Module/AnimNextModule_EditorData.h
--- a/Engine/Plugins/Experimental/UAF/UAF/Source/UAFUncookedOnly/Internal/Module/AnimNextModule_EditorData.h
+++ b/Engine/Plugins/Experimental/UAF/UAF/Source/UAFUncookedOnly/Internal/Module/AnimNextModule_EditorData.h
@@ -31,11 +31,11 @@ namespace UE::UAF::Editor
 
 /** Editor data for AnimNext modules */
 UCLASS(MinimalAPI)
-class UAnimNextModule_EditorData : public UAnimNextSharedVariables_EditorData
+class UUAFSystem_EditorData : public UAnimNextSharedVariables_EditorData
 {
 	GENERATED_BODY()
 
-	friend class UAnimNextModuleFactory;
+	friend class UUAFSystemFactory;
 	friend class UAnimNextEdGraph;
 	friend struct UE::UAF::UncookedOnly::FUtils;
 	friend struct UE::UAF::Editor::FUtils;
diff --git a/Engine/Plugins/Experimental/UAF/UAF/Source/UAFUncookedOnly/Internal/UncookedOnlyUtils.h b/Engine/Plugins/Experimental/UAF/UAF/Source/UAFUncookedOnly/Internal/UncookedOnlyUtils.h
--- a/Engine/Plugins/Experimental/UAF/UAF/Source/UAFUncookedOnly/Internal/UncookedOnlyUtils.h
+++ b/Engine/Plugins/Experimental/UAF/UAF/Source/UAFUncookedOnly/Internal/UncookedOnlyUtils.h
@@ -31,7 +31,7 @@ struct FRigVMGraphFunctionHeaderArray;
 struct FGraphContextMenuBuilder;
 struct FRigVMFunction;
 class UUAFSystem;
-class UAnimNextModule_EditorData;
+class UUAFSystem_EditorData;
 class UAnimNextEdGraph;
 class URigVMController;
 class URigVMGraph;
diff --git a/Engine/Plugins/Experimental/UAF/UAF/Source/UAFUncookedOnly/Private/Module/AnimNextModule_EditorData.cpp b/Engine/Plugins/Experimental/UAF/UAF/Source/UAFUncookedOnly/Private/Module/AnimNextModule_EditorData.cpp
--- a/Engine/Plugins/Experimental/UAF/UAF/Source/UAFUncookedOnly/Private/Module/AnimNextModule_EditorData.cpp
+++ b/Engine/Plugins/Experimental/UAF/UAF/Source/UAFUncookedOnly/Private/Module/AnimNextModule_EditorData.cpp
@@ -17,7 +17,7 @@
 
 #include UE_INLINE_GENERATED_CPP_BY_NAME(AnimNextModule_EditorData)
 
-TConstArrayView<TSubclassOf<UAnimNextRigVMAssetEntry>> UAnimNextModule_EditorData::GetEntryClasses() const
+TConstArrayView<TSubclassOf<UAnimNextRigVMAssetEntry>> UUAFSystem_EditorData::GetEntryClasses() const
 {
 	static const TSubclassOf<UAnimNextRigVMAssetEntry> Classes[] =
 	{
@@ -29,7 +29,7 @@ TConstArrayView<TSubclassOf<UAnimNextRigVMAssetEntry>> UAnimNextModule_EditorDat
 	return Classes;
 }
 
-void UAnimNextModule_EditorData::OnPreCompileAsset(FRigVMCompileSettings& InSettings)
+void UUAFSystem_EditorData::OnPreCompileAsset(FRigVMCompileSettings& InSettings)
 {
 	using namespace UE::UAF::UncookedOnly;
 
@@ -38,14 +38,14 @@ void UAnimNextModule_EditorData::OnPreCompileAsset(FRigVMCompileSettings& InSett
 	Module->Dependencies.Empty();
 }
 
-void UAnimNextModule_EditorData::OnPreCompileGetProgrammaticGraphs(const FRigVMCompileSettings& InSettings, FAnimNextGetGraphCompileContext& OutCompileContext)
+void UUAFSystem_EditorData::OnPreCompileGetProgrammaticGraphs(const FRigVMCompileSettings& InSettings, FAnimNextGetGraphCompileContext& OutCompileContext)
 {
 	using namespace UE::UAF::UncookedOnly;
 
 	FUtils::CompileVariableBindings(InSettings, FUtils::GetAsset<UUAFSystem>(this), OutCompileContext.GetMutableProgrammaticGraphs());
 }
 
-void UAnimNextModule_EditorData::OnPreCompileProcessGraphs(const FRigVMCompileSettings& InSettings, FAnimNextProcessGraphCompileContext& OutCompileContext)
+void UUAFSystem_EditorData::OnPreCompileProcessGraphs(const FRigVMCompileSettings& InSettings, FAnimNextProcessGraphCompileContext& OutCompileContext)
 {
 	using namespace UE::UAF::UncookedOnly;
 
@@ -55,7 +55,7 @@ void UAnimNextModule_EditorData::OnPreCompileProcessGraphs(const FRigVMCompileSe
 	Module->Dependencies = Dependencies;
 }
 
-void UAnimNextModule_EditorData::CustomizeNewAssetEntry(UAnimNextRigVMAssetEntry* InNewEntry) const
+void UUAFSystem_EditorData::CustomizeNewAssetEntry(UAnimNextRigVMAssetEntry* InNewEntry) const
 {
 	Super::CustomizeNewAssetEntry(InNewEntry);
 	
diff --git a/Engine/Plugins/Experimental/UAF/UAF/Source/UAFUncookedOnly/Private/UncookedOnlyUtils.cpp b/Engine/Plugins/Experimental/UAF/UAF/Source/UAFUncookedOnly/Private/UncookedOnlyUtils.cpp
--- a/Engine/Plugins/Experimental/UAF/UAF/Source/UAFUncookedOnly/Private/UncookedOnlyUtils.cpp
+++ b/Engine/Plugins/Experimental/UAF/UAF/Source/UAFUncookedOnly/Private/UncookedOnlyUtils.cpp
@@ -395,7 +395,7 @@ void FUtils::CompileVariableBindingsInternal(const FRigVMCompileSettings& InSett
 	}
 
 	const bool bHasBindings = BindingGroups.Num() > 0;
-	const bool bHasPublicVariablesToCopy = EditorData->IsA<UAnimNextModule_EditorData>() && EditorData->HasPublicVariables() && bInThreadSafe;
+	const bool bHasPublicVariablesToCopy = EditorData->IsA<UUAFSystem_EditorData>() && EditorData->HasPublicVariables() && bInThreadSafe;
 	if(!bHasBindings && !bHasPublicVariablesToCopy)
 	{
 		// Nothing to do here
diff --git a/Engine/Plugins/Experimental/UAF/UAFTests/Source/UAFTestSuite/Private/AnimNextEditorTest.cpp b/Engine/Plugins/Experimental/UAF/UAFTests/Source/UAFTestSuite/Private/AnimNextEditorTest.cpp
--- a/Engine/Plugins/Experimental/UAF/UAFTests/Source/UAFTestSuite/Private/AnimNextEditorTest.cpp
+++ b/Engine/Plugins/Experimental/UAF/UAFTests/Source/UAFTestSuite/Private/AnimNextEditorTest.cpp
@@ -39,7 +39,7 @@ bool FEditor_Variables::RunTest(const FString& InParameters)
 	
 	FFactoryAndClass FactoryClassPairs[] =
 	{
-		{ UAnimNextModuleFactory::StaticClass(), UUAFSystem::StaticClass() },
+		{ UUAFSystemFactory::StaticClass(), UUAFSystem::StaticClass() },
 	};
 
 	for(const FFactoryAndClass& FactoryAndClass : FactoryClassPairs)
@@ -141,7 +141,7 @@ bool FEditor_Graphs::RunTest(const FString& InParameters)
 	
 	FTestSettings TestSettings[] =
 	{
-		{ UAnimNextModuleFactory::StaticClass(), UUAFSystem::StaticClass(), true, true, false, false },
+		{ UUAFSystemFactory::StaticClass(), UUAFSystem::StaticClass(), true, true, false, false },
 	};
 
 	for(const FTestSettings& TestSetting : TestSettings)
@@ -161,12 +161,12 @@ bool FEditor_Graphs::RunTest(const FString& InParameters)
 			UAnimNextEventGraphEntry* EventGraphEntry = nullptr;
 			if(TestSetting.bExpectExistingEventGraph)
 			{
-				EventGraphEntry = Cast<UAnimNextEventGraphEntry>(EditorData->FindEntry(UAnimNextModule_EditorData::DefaultEventGraphName));
+				EventGraphEntry = Cast<UAnimNextEventGraphEntry>(EditorData->FindEntry(UUAFSystem_EditorData::DefaultEventGraphName));
 				UE_RETURN_ON_ERROR(EventGraphEntry != nullptr, TEXT("Could not find existing event graph."));
 			}
 			else
 			{
-				EventGraphEntry = EditorData->AddEventGraph(UAnimNextModule_EditorData::DefaultEventGraphName, FRigUnit_AnimNextPrePhysicsEvent::StaticStruct());
+				EventGraphEntry = EditorData->AddEventGraph(UUAFSystem_EditorData::DefaultEventGraphName, FRigUnit_AnimNextPrePhysicsEvent::StaticStruct());
 				UE_RETURN_ON_ERROR(EventGraphEntry != nullptr, TEXT("Could not add event graph."));
 			}
 
@@ -181,13 +181,13 @@ bool FEditor_Graphs::RunTest(const FString& InParameters)
 
 			GEditor->UndoTransaction();
 
-			UAnimNextEventGraphEntry* FoundEventGraphEntry = Cast<UAnimNextEventGraphEntry>(EditorData->FindEntry(UAnimNextModule_EditorData::DefaultEventGraphName));
+			UAnimNextEventGraphEntry* FoundEventGraphEntry = Cast<UAnimNextEventGraphEntry>(EditorData->FindEntry(UUAFSystem_EditorData::DefaultEventGraphName));
 			UE_RETURN_ON_ERROR(FoundEventGraphEntry != nullptr, "FEditor_Graphs -> Could not find event graph post-undo.");
 		}
 		else
 		{
 			AddExpectedError(TEXT("Cannot add an event graph to this asset - entry is not allowed"));
-			UAnimNextEventGraphEntry* EventGraphEntry = EditorData->AddEventGraph(UAnimNextModule_EditorData::DefaultEventGraphName, FRigUnit_AnimNextPrePhysicsEvent::StaticStruct());
+			UAnimNextEventGraphEntry* EventGraphEntry = EditorData->AddEventGraph(UUAFSystem_EditorData::DefaultEventGraphName, FRigUnit_AnimNextPrePhysicsEvent::StaticStruct());
 		}
 	}
 
diff --git a/Engine/Plugins/Experimental/UAF/UAFTests/Source/UAFTestSuite/Private/AnimNextVariablesTest.cpp b/Engine/Plugins/Experimental/UAF/UAFTests/Source/UAFTestSuite/Private/AnimNextVariablesTest.cpp
--- a/Engine/Plugins/Experimental/UAF/UAFTests/Source/UAFTestSuite/Private/AnimNextVariablesTest.cpp
+++ b/Engine/Plugins/Experimental/UAF/UAFTests/Source/UAFTestSuite/Private/AnimNextVariablesTest.cpp
@@ -228,7 +228,7 @@ bool FVariables::RunTest(const FString& InParameters)
 {
 	ON_SCOPE_EXIT{ FUtils::CleanupAfterTests(); };
 
-	UFactory* Factory = NewObject<UFactory>(GetTransientPackage(), UAnimNextModuleFactory::StaticClass());
+	UFactory* Factory = NewObject<UFactory>(GetTransientPackage(), UUAFSystemFactory::StaticClass());
 	UAnimNextRigVMAsset* Asset = CastChecked<UAnimNextRigVMAsset>(Factory->FactoryCreateNew(UUAFSystem::StaticClass(), GetTransientPackage(), TEXT("TestAsset"), RF_Transient, nullptr, nullptr, NAME_None));
 	UE_RETURN_ON_ERROR(Asset != nullptr, "FVariables -> Failed to create asset");
 
@@ -244,13 +244,13 @@ bool FVariables::RunTest(const FString& InParameters)
 	UE_RETURN_ON_ERROR(ResultEntry != nullptr, TEXT("Could not create new variable in graph."));
 
 	// Get event graph
-	UAnimNextEventGraphEntry* EventGraph = Cast<UAnimNextEventGraphEntry>(EditorData->FindEntry(UAnimNextModule_EditorData::DefaultEventGraphName));
+	UAnimNextEventGraphEntry* EventGraph = Cast<UAnimNextEventGraphEntry>(EditorData->FindEntry(UUAFSystem_EditorData::DefaultEventGraphName));
 	URigVMGraph* RigVMGraph = nullptr;
 	URigVMNode* PrePhysicsEventNode = nullptr;
 	
 	if(EventGraph == nullptr)
 	{
-		EventGraph = EditorData->AddEventGraph(UAnimNextModule_EditorData::DefaultEventGraphName, FRigUnit_AnimNextPrePhysicsEvent::StaticStruct());
+		EventGraph = EditorData->AddEventGraph(UUAFSystem_EditorData::DefaultEventGraphName, FRigUnit_AnimNextPrePhysicsEvent::StaticStruct());
 		UE_RETURN_ON_ERROR(EventGraph != nullptr, TEXT("Could not create new event graph in asset."));
 
 		RigVMGraph = EventGraph->GetRigVMGraph();
@@ -327,7 +327,7 @@ bool FVariables_UOLBindings::RunTest(const FString& InParameters)
 {
 	ON_SCOPE_EXIT{ FUtils::CleanupAfterTests(); };
 
-	UFactory* Factory = NewObject<UFactory>(GetTransientPackage(), UAnimNextModuleFactory::StaticClass());
+	UFactory* Factory = NewObject<UFactory>(GetTransientPackage(), UUAFSystemFactory::StaticClass());
 	UUAFSystem* Asset = CastChecked<UUAFSystem>(Factory->FactoryCreateNew(UUAFSystem::StaticClass(), GetTransientPackage(), TEXT("TestAsset"), RF_Transient, nullptr, nullptr, NAME_None));
 	UE_RETURN_ON_ERROR(Asset != nullptr, "FVariables_UOLBindings -> Failed to create asset");
 
@@ -372,13 +372,13 @@ bool FVariables_UOLBindings::RunTest(const FString& InParameters)
 	HoistedAccessorEntry->SetBinding(MoveTemp(BindingC));
 
 	// Get event graph
-	UAnimNextEventGraphEntry* EventGraph = Cast<UAnimNextEventGraphEntry>(EditorData->FindEntry(UAnimNextModule_EditorData::DefaultEventGraphName));
+	UAnimNextEventGraphEntry* EventGraph = Cast<UAnimNextEventGraphEntry>(EditorData->FindEntry(UUAFSystem_EditorData::DefaultEventGraphName));
 	URigVMGraph* RigVMGraph = nullptr;
 	URigVMNode* PrePhysicsEventNode = nullptr;
 	
 	if(EventGraph == nullptr)
 	{
-		EventGraph = EditorData->AddEventGraph(UAnimNextModule_EditorData::DefaultEventGraphName, FRigUnit_AnimNextPrePhysicsEvent::StaticStruct());
+		EventGraph = EditorData->AddEventGraph(UUAFSystem_EditorData::DefaultEventGraphName, FRigUnit_AnimNextPrePhysicsEvent::StaticStruct());
 		UE_RETURN_ON_ERROR(EventGraph != nullptr, TEXT("Could not create new event graph in asset."));
 
 		RigVMGraph = EventGraph->GetRigVMGraph();
@@ -495,7 +495,7 @@ bool FVariables_Shared::RunTest(const FString& InParameters)
 {
 	ON_SCOPE_EXIT{ FUtils::CleanupAfterTests(); };
 
-	UFactory* Factory = NewObject<UFactory>(GetTransientPackage(), UAnimNextModuleFactory::StaticClass());
+	UFactory* Factory = NewObject<UFactory>(GetTransientPackage(), UUAFSystemFactory::StaticClass());
 	UUAFSystem* Asset = CastChecked<UUAFSystem>(Factory->FactoryCreateNew(UUAFSystem::StaticClass(), GetTransientPackage(), TEXT("TestAsset"), RF_Transient, nullptr, nullptr, NAME_None));
 	UE_RETURN_ON_ERROR(Asset != nullptr, "FVariables_Shared -> Failed to create asset");
 
diff --git a/Engine/Plugins/Experimental/UAF/UAFTests/Source/UAFTests/Private/UAFAssetCreationTests.cpp b/Engine/Plugins/Experimental/UAF/UAFTests/Source/UAFTests/Private/UAFAssetCreationTests.cpp
--- a/Engine/Plugins/Experimental/UAF/UAFTests/Source/UAFTests/Private/UAFAssetCreationTests.cpp
+++ b/Engine/Plugins/Experimental/UAF/UAFTests/Source/UAFTests/Private/UAFAssetCreationTests.cpp
@@ -65,8 +65,8 @@ TEST_CLASS(FAssetCreationTests, "Animation.UAF.Functional")
 		TestCommandBuilder
 			.Do(TEXT("Create AnimNext Module Asset"), [&]()
 			{
-				const TSubclassOf<UAnimNextModuleFactory> ModuleFactoryClass = UAnimNextModuleFactory::StaticClass();
-				UObject* FactoryObject = UAFTestsUtilities::CreateFactoryObject(NewObject<UAnimNextModuleFactory>(), UUAFSystem::StaticClass(), "NewAnimNextModule");				
+				const TSubclassOf<UUAFSystemFactory> ModuleFactoryClass = UUAFSystemFactory::StaticClass();
+				UObject* FactoryObject = UAFTestsUtilities::CreateFactoryObject(NewObject<UUAFSystemFactory>(), UUAFSystem::StaticClass(), "NewAnimNextModule");				
 				AnimNextAsset = CastChecked<UUAFSystem>(FactoryObject);
 			})
 			.Then(TEXT("Assert AnimNext Module Asset Created and Open in Editor"), [&]()
diff --git a/Engine/Plugins/Experimental/UAF/UAFTests/Source/UAFTests/Private/UAFBlueprintTests.cpp b/Engine/Plugins/Experimental/UAF/UAFTests/Source/UAFTests/Private/UAFBlueprintTests.cpp
--- a/Engine/Plugins/Experimental/UAF/UAFTests/Source/UAFTests/Private/UAFBlueprintTests.cpp
+++ b/Engine/Plugins/Experimental/UAF/UAFTests/Source/UAFTests/Private/UAFBlueprintTests.cpp
@@ -66,8 +66,8 @@ TEST_CLASS(FAssetBlueprintTests, "Animation.UAF.Functional")
 			.Do(TEXT("Create UAF Module Asset"), [&]()
 			{
 				const FString UAFModuleName = "NewAnimNextModule";
-				const TSubclassOf<UAnimNextModuleFactory> ModuleFactoryClass = UAnimNextModuleFactory::StaticClass();
-				UObject* factoryObject = UAFTestsUtilities::CreateFactoryObject(NewObject<UAnimNextModuleFactory>(), UUAFSystem::StaticClass(), UAFModuleName);				
+				const TSubclassOf<UUAFSystemFactory> ModuleFactoryClass = UUAFSystemFactory::StaticClass();
+				UObject* factoryObject = UAFTestsUtilities::CreateFactoryObject(NewObject<UUAFSystemFactory>(), UUAFSystem::StaticClass(), UAFModuleName);				
 				UAFAsset = CastChecked<UUAFSystem>(factoryObject);			
 				ASSERT_THAT(IsNotNull(UAFAsset));
 				ASSERT_THAT(AreEqual(UAFModuleName, UAFAsset->GetFName()));
diff --git a/Engine/Plugins/Experimental/UAF/UAFTests/Source/UAFTests/Private/UAFFunctionNodeTests.cpp b/Engine/Plugins/Experimental/UAF/UAFTests/Source/UAFTests/Private/UAFFunctionNodeTests.cpp
--- a/Engine/Plugins/Experimental/UAF/UAFTests/Source/UAFTests/Private/UAFFunctionNodeTests.cpp
+++ b/Engine/Plugins/Experimental/UAF/UAFTests/Source/UAFTests/Private/UAFFunctionNodeTests.cpp
@@ -70,8 +70,8 @@ TEST_CLASS(FAssetFunctionNodeTests, "Animation.UAF.Functional")
 		TestCommandBuilder
 			.Do(TEXT("Create AnimNext Module Asset"), [&]()
 			{
-				const TSubclassOf<UAnimNextModuleFactory> AnimNextModuleFactoryClass = UAnimNextModuleFactory::StaticClass();
-				UObject* factoryObject = UAFTestsUtilities::CreateFactoryObject(NewObject<UAnimNextModuleFactory>(GetTransientPackage(), AnimNextModuleFactoryClass.Get()), UUAFSystem::StaticClass(), "NewAnimNextModule");
+				const TSubclassOf<UUAFSystemFactory> AnimNextModuleFactoryClass = UUAFSystemFactory::StaticClass();
+				UObject* factoryObject = UAFTestsUtilities::CreateFactoryObject(NewObject<UUAFSystemFactory>(GetTransientPackage(), AnimNextModuleFactoryClass.Get()), UUAFSystem::StaticClass(), "NewAnimNextModule");
 				AnimNextAsset = CastChecked<UUAFSystem>(factoryObject);				
 				ASSERT_THAT(IsTrue(AssetEditorSubsystem->OpenEditorForAsset(AnimNextAsset)));
 			})
diff --git a/Engine/Plugins/Experimental/UAF/UAFTests/Source/UAFTests/Private/UAFModuleTests.cpp b/Engine/Plugins/Experimental/UAF/UAFTests/Source/UAFTests/Private/UAFModuleTests.cpp
--- a/Engine/Plugins/Experimental/UAF/UAFTests/Source/UAFTests/Private/UAFModuleTests.cpp
+++ b/Engine/Plugins/Experimental/UAF/UAFTests/Source/UAFTests/Private/UAFModuleTests.cpp
@@ -278,7 +278,7 @@ TEST_CLASS(FAnimNextModuleTests, "Animation.UAF.Functional")
 protected:
 	void CreateAnimNextModule()
 	{
-		UFactory* Factory = NewObject<UAnimNextModuleFactory>();
+		UFactory* Factory = NewObject<UUAFSystemFactory>();
 		AnimNextModulePackage = CreateUniquePackage("NewAnimNextModule");
 		AnimNextModulePackageName = *FPaths::GetBaseFilename(AnimNextModulePackage->GetName());		
 		UObject* FactoryObject = Factory->FactoryCreateNew(UUAFSystem::StaticClass(), AnimNextModulePackage, AnimNextModulePackageName, RF_Public | RF_Standalone, NULL, GWarn);
@@ -309,7 +309,7 @@ bool FSystemReference_Creation::RunTest(const FString& InParameters)
 
 	ON_SCOPE_EXIT{ FUtils::CleanupAfterTests(); };
 
-	UFactory* Factory = NewObject<UFactory>(GetTransientPackage(), UAnimNextModuleFactory::StaticClass());
+	UFactory* Factory = NewObject<UFactory>(GetTransientPackage(), UUAFSystemFactory::StaticClass());
 	UUAFSystem* Asset = CastChecked<UUAFSystem>(Factory->FactoryCreateNew(UUAFSystem::StaticClass(), GetTransientPackage(), TEXT("TestAsset"), RF_Transient, nullptr, nullptr, NAME_None));
 	UE_RETURN_ON_ERROR(Asset != nullptr, "FSystemReference_Creation -> Failed to create asset");
 
diff --git a/Engine/Plugins/Experimental/UAF/UAFTests/Source/UAFTests/Private/UAFTestBlueprintLibrary.h b/Engine/Plugins/Experimental/UAF/UAFTests/Source/UAFTests/Private/UAFTestBlueprintLibrary.h
--- a/Engine/Plugins/Experimental/UAF/UAFTests/Source/UAFTests/Private/UAFTestBlueprintLibrary.h
+++ b/Engine/Plugins/Experimental/UAF/UAFTests/Source/UAFTests/Private/UAFTestBlueprintLibrary.h
@@ -205,7 +205,7 @@ public:
 			return;
 		}
 	
-		if (UAnimNextModule_EditorData* EditorData = UE::UAF::UncookedOnly::FUtils::GetEditorData<UAnimNextModule_EditorData>(InModule))
+		if (UUAFSystem_EditorData* EditorData = UE::UAF::UncookedOnly::FUtils::GetEditorData<UUAFSystem_EditorData>(InModule))
 		{
 			if (EditorData->bVMRecompilationRequired)
 			{
