Fix missing python scripts during rename

#Synced-CL 49648897

diff --git a/Engine/Plugins/Experimental/UAF/UAFTests/Source/UAFTestSuite/Private/AnimNextEditorTest.cpp b/Engine/Plugins/Experimental/UAF/UAFTests/Source/UAFTestSuite/Private/AnimNextEditorTest.cpp
--- a/Engine/Plugins/Experimental/UAF/UAFTests/Source/UAFTestSuite/Private/AnimNextEditorTest.cpp
+++ b/Engine/Plugins/Experimental/UAF/UAFTests/Source/UAFTestSuite/Private/AnimNextEditorTest.cpp
@@ -203,7 +203,7 @@ bool FEditor_Variables_Python::RunTest(const FString& InParameters)
 
 	const TCHAR* Script = TEXT(
 		"asset_tools = unreal.AssetToolsHelpers.get_asset_tools()\n"
-		"animation_module = unreal.AssetTools.create_asset(asset_tools, asset_name = \"TestModule\", package_path = \"/Game/\", asset_class = unreal.AnimNextModule, factory = unreal.AnimNextModuleFactory())\n"
+		"animation_module = unreal.AssetTools.create_asset(asset_tools, asset_name = \"TestModule\", package_path = \"/Game/\", asset_class = unreal.UAFSystem, factory = unreal.UAFSystemFactory())\n"
 		"animation_module.add_variable(name = \"TestParam\", value_type = unreal.PropertyBagPropertyType.BOOL, container_type = unreal.PropertyBagContainerType.NONE)\n"
 		"unreal.EditorAssetLibrary.delete_loaded_asset(animation_module)\n"
 	);
