User can add a shared variable reference to a UAF asset.
User can set a shared variables access level.
Automated QMetry Test Case: UE-TC-21565
Automated QMetry Test Case: UE-TC-21572

#Synced-CL 49657956

diff --git a/Engine/Plugins/Experimental/UAF/UAF/Source/UAFUncookedOnly/Internal/Entries/AnimNextVariableEntry.h b/Engine/Plugins/Experimental/UAF/UAF/Source/UAFUncookedOnly/Internal/Entries/AnimNextVariableEntry.h
--- a/Engine/Plugins/Experimental/UAF/UAF/Source/UAFUncookedOnly/Internal/Entries/AnimNextVariableEntry.h
+++ b/Engine/Plugins/Experimental/UAF/UAF/Source/UAFUncookedOnly/Internal/Entries/AnimNextVariableEntry.h
@@ -89,6 +89,7 @@ class UAnimNextVariableEntry : public UAnimNextRigVMAssetEntry, public IAnimNext
 	friend class UE::UAF::Editor::FVariablesOutlinerMode;
 	friend class UE::UAF::UncookedOnly::FModule;
 	friend class UE::UAF::Tests::FVariableReferenceTest;
+	friend struct FAssetVariableTests;
 
 
 	UAnimNextVariableEntry();
diff --git a/Engine/Plugins/Experimental/UAF/UAFTests/Source/UAFTests/Private/UAFTestsUtilities.cpp b/Engine/Plugins/Experimental/UAF/UAFTests/Source/UAFTests/Private/UAFTestsUtilities.cpp
--- a/Engine/Plugins/Experimental/UAF/UAFTests/Source/UAFTests/Private/UAFTestsUtilities.cpp
+++ b/Engine/Plugins/Experimental/UAF/UAFTests/Source/UAFTests/Private/UAFTestsUtilities.cpp
@@ -83,7 +83,25 @@ namespace UAFTestsUtilities
 		return EditorData->AddVariable(FName(InVariableName), InVariableType, InDefaultValue);
 	}
 
-		UEdGraphNode* AddVariableNode(UEdGraph* InParentGraph, const UObject* InSourceObject, const FString& InVariableName, const FAnimNextParamType& InVariableType, const FAnimNextSchemaAction_Variable::EVariableAccessorChoice InVariableAccessorChoice, TArray<UEdGraphPin*>& InFromPins, const FVector2f& InVariableLocation)
+	UAnimNextSharedVariablesEntry* AddSharedVariables(UAnimNextRigVMAsset* InAnimNextAsset, UAnimNextSharedVariables* InSharedVariables)
+	{
+		if (!InAnimNextAsset)
+		{
+			UE_LOG(LogTemp, Error, TEXT("Invalid UAnimNextRigVMAsset instance provided."));
+			return nullptr;
+		}
+
+		if (!InSharedVariables)
+		{
+			UE_LOG(LogTemp, Error, TEXT("Invalid UAnimNextSharedVariables instance provided."));
+			return nullptr;
+		}
+		
+		UAnimNextRigVMAssetEditorData* EditorData = UE::UAF::UncookedOnly::FUtils::GetEditorData(InAnimNextAsset);
+		return EditorData->AddSharedVariables(InSharedVariables);
+	}
+
+	UEdGraphNode* AddVariableNode(UEdGraph* InParentGraph, const UObject* InSourceObject, const FString& InVariableName, const FAnimNextParamType& InVariableType, const FAnimNextSchemaAction_Variable::EVariableAccessorChoice InVariableAccessorChoice, TArray<UEdGraphPin*>& InFromPins, const FVector2f& InVariableLocation)
 	{
 		if (!InParentGraph)
 		{
diff --git a/Engine/Plugins/Experimental/UAF/UAFTests/Source/UAFTests/Private/UAFVariableTests.cpp b/Engine/Plugins/Experimental/UAF/UAFTests/Source/UAFTests/Private/UAFVariableTests.cpp
--- a/Engine/Plugins/Experimental/UAF/UAFTests/Source/UAFTests/Private/UAFVariableTests.cpp
+++ b/Engine/Plugins/Experimental/UAF/UAFTests/Source/UAFTests/Private/UAFVariableTests.cpp
@@ -9,13 +9,20 @@
 #include "Subsystems/AssetEditorSubsystem.h"
 #include "Graph/AnimNextAnimationGraph.h"
 #include "Graph/AnimNextAnimationGraphFactory.h"
+#include "Module/AnimNextModule.h"
+#include "Module/AnimNextModuleFactory.h"
+#include "Variables/AnimNextSharedVariables.h"
+#include "Variables/AnimNextSharedVariablesFactory.h"
 #include "AnimNextRigVMAssetEditorData.h"
 
 TEST_CLASS(FAssetVariableTests, "Animation.UAF.Functional")
 {
-	UAnimNextRigVMAsset* AnimNextAsset;
-	UAssetEditorSubsystem* AssetEditorSubsystem;
-	UAnimNextRigVMAssetEditorData* EditorData;
+	UAnimNextRigVMAsset* AnimNextAsset = nullptr;
+	UAnimNextSharedVariables* AnimNextSharedVariablesAsset = nullptr;
+	UAnimNextVariableEntry* VariableEntry = nullptr;
+	UAnimNextSharedVariablesEntry* SharedVariablesEntry = nullptr;
+	UAssetEditorSubsystem* AssetEditorSubsystem = nullptr;
+	UAnimNextRigVMAssetEditorData* EditorData = nullptr;
 
 	BEFORE_EACH()
 	{
@@ -30,17 +37,101 @@ TEST_CLASS(FAssetVariableTests, "Animation.UAF.Functional")
 			.Do(TEXT("Create AnimNext Animation Graph Asset"), [&]()
 			{
 				const TSubclassOf<UUAFAnimGraphFactory> AnimNextAnimGraphFactoryClass = UUAFAnimGraphFactory::StaticClass();
-				UObject* factoryObject = UAFTestsUtilities::CreateFactoryObject(NewObject<UUAFAnimGraphFactory>(GetTransientPackage(), AnimNextAnimGraphFactoryClass.Get()), UUAFAnimGraph::StaticClass(), "NewAnimNextAnimationGraph");
+				UObject* factoryObject = UAFTestsUtilities::CreateFactoryObject(NewObject<UUAFAnimGraphFactory>(GetTransientPackage(), AnimNextAnimGraphFactoryClass.Get()), UUAFAnimGraph::StaticClass(), "NewUAFAnimGraph");
 				AnimNextAsset = CastChecked<UUAFAnimGraph>(factoryObject);				
 				ASSERT_THAT(IsTrue(AssetEditorSubsystem->OpenEditorForAsset(AnimNextAsset)));
 			})
 			.Then(TEXT("Create Bool Variable"), [&]()
 			{
 				FAnimNextParamType Type = FAnimNextParamType::GetType<bool>();
-				UAnimNextVariableEntry* VariableEntry = UAFTestsUtilities::AddVariable(AnimNextAsset, Type, "NewVariable", "false");
+				VariableEntry = UAFTestsUtilities::AddVariable(AnimNextAsset, Type, "NewVariable", "false");
 				ASSERT_THAT(IsNotNull(VariableEntry));
 			});
-	}	
+	}
+	
+	// QMetry Test Case:  UE-TC-21565
+	TEST_METHOD(UAFSystem_Add_UAFSharedVariables)
+	{
+		TestCommandBuilder
+			.Do(TEXT("Create UAF Shared Variables Asset"), [&]()
+			{
+				const TSubclassOf<UAnimNextSharedVariablesFactory> SharedVariablesFactoryClass = UAnimNextSharedVariablesFactory::StaticClass();
+				UObject* FactoryObject = UAFTestsUtilities::CreateFactoryObject(NewObject<UAnimNextSharedVariablesFactory>(), UAnimNextSharedVariables::StaticClass(), "NewAnimNextSharedVariables");
+				AnimNextSharedVariablesAsset = CastChecked<UAnimNextSharedVariables>(FactoryObject);
+				ASSERT_THAT(IsNotNull(AnimNextSharedVariablesAsset));
+				ASSERT_THAT(IsTrue(AssetEditorSubsystem->OpenEditorForAsset(AnimNextSharedVariablesAsset)));		
+			})
+			.Then(TEXT("Add Variables to UAF Shared Variables Asset"), [&]()
+			{
+				UAnimNextRigVMAsset* Asset = CastChecked<UAnimNextRigVMAsset>(AnimNextSharedVariablesAsset);
+				FAnimNextParamType Type = FAnimNextParamType::GetType<bool>();
+				VariableEntry = UAFTestsUtilities::AddVariable(Asset, Type, "NewVariable", "false");
+				ASSERT_THAT(IsNotNull(VariableEntry));
+				
+				Type = FAnimNextParamType::GetType<FString>();
+				VariableEntry = UAFTestsUtilities::AddVariable(Asset, Type, "NewVariable", "");
+				ASSERT_THAT(IsNotNull(VariableEntry));
+			})
+			.Then(TEXT("Create UAF System Asset"), [&]()
+			{
+				const TSubclassOf<UUAFSystemFactory> ModuleFactoryClass = UUAFSystemFactory::StaticClass();
+				UObject* FactoryObject = UAFTestsUtilities::CreateFactoryObject(NewObject<UUAFSystemFactory>(), UUAFSystem::StaticClass(), "NewUAFSystem");				
+				AnimNextAsset = CastChecked<UUAFSystem>(FactoryObject);
+				ASSERT_THAT(IsNotNull(AnimNextAsset));
+				ASSERT_THAT(IsTrue(AssetEditorSubsystem->OpenEditorForAsset(AnimNextAsset)));
+			})		
+			.Then(TEXT("Add UAF Shared Variables Asset to UAF System Asset"), [&]()
+			{
+				SharedVariablesEntry = UAFTestsUtilities::AddSharedVariables(AnimNextAsset, AnimNextSharedVariablesAsset);
+				ASSERT_THAT(IsNotNull(SharedVariablesEntry));
+			})
+			.Then(TEXT("Add Variable to UAF System Asset"), [&]()
+			{
+				FAnimNextParamType Type = FAnimNextParamType::GetType<bool>();
+				VariableEntry = UAFTestsUtilities::AddVariable(AnimNextAsset, Type, "NewVariable_1", "false");
+				ASSERT_THAT(IsNotNull(VariableEntry));
+			});
+	}
+	
+	// QMetry Test Case:  UE-TC-21572
+	TEST_METHOD(UAFSharedVariables_Public_Private_Access)
+	{
+		TestCommandBuilder
+			.Do(TEXT("Create UAF System Asset"), [&]()
+			{
+				const TSubclassOf<UUAFSystemFactory> ModuleFactoryClass = UUAFSystemFactory::StaticClass();
+				UObject* FactoryObject = UAFTestsUtilities::CreateFactoryObject(NewObject<UUAFSystemFactory>(), UUAFSystem::StaticClass(), "NewUAFSystem");				
+				AnimNextSharedVariablesAsset = CastChecked<UUAFSystem>(FactoryObject);
+				ASSERT_THAT(IsNotNull(AnimNextSharedVariablesAsset));
+				ASSERT_THAT(IsTrue(AssetEditorSubsystem->OpenEditorForAsset(AnimNextSharedVariablesAsset)));
+			})
+			.Then(TEXT("Add Variable to UAF System Asset"), [&]()
+			{
+				FAnimNextParamType Type = FAnimNextParamType::GetType<bool>();
+				VariableEntry = UAFTestsUtilities::AddVariable(AnimNextSharedVariablesAsset, Type, "NewVariable", "false");
+				ASSERT_THAT(IsNotNull(VariableEntry));
+				ASSERT_THAT(AreEqual(EAnimNextExportAccessSpecifier::Public, VariableEntry->Access));
+			})
+			.Then(TEXT("Create UAF Anim Graph Asset"), [&]()
+			{
+				const TSubclassOf<UUAFAnimGraphFactory> AnimGraphFactoryClass = UUAFAnimGraphFactory::StaticClass();
+				UObject* FactoryObject = UAFTestsUtilities::CreateFactoryObject(NewObject<UUAFAnimGraphFactory>(), UUAFAnimGraph::StaticClass(), "NewUAFAnimGraph");
+				AnimNextAsset = CastChecked<UUAFAnimGraph>(FactoryObject);
+				ASSERT_THAT(IsNotNull(AnimNextAsset));
+				ASSERT_THAT(IsTrue(AssetEditorSubsystem->OpenEditorForAsset(AnimNextAsset)));
+			})		
+			.Then(TEXT("Add UAF Shared Variables Asset to UAF System Asset"), [&]()
+			{
+				SharedVariablesEntry = UAFTestsUtilities::AddSharedVariables(AnimNextAsset, AnimNextSharedVariablesAsset);
+				ASSERT_THAT(IsNotNull(SharedVariablesEntry));
+			})
+			.Then(TEXT("Add Variable to UAF System Asset"), [&]()
+			{
+				VariableEntry->Access = EAnimNextExportAccessSpecifier::Private;
+				ASSERT_THAT(AreEqual(EAnimNextExportAccessSpecifier::Private, VariableEntry->Access));
+			});
+	}
+	
 };
 
 #endif // WITH_EDITOR
diff --git a/Engine/Plugins/Experimental/UAF/UAFTests/Source/UAFTests/Public/UAFTestsUtilities.h b/Engine/Plugins/Experimental/UAF/UAFTests/Source/UAFTests/Public/UAFTestsUtilities.h
--- a/Engine/Plugins/Experimental/UAF/UAFTests/Source/UAFTests/Public/UAFTestsUtilities.h
+++ b/Engine/Plugins/Experimental/UAF/UAFTests/Source/UAFTests/Public/UAFTestsUtilities.h
@@ -57,6 +57,15 @@ namespace UAFTestsUtilities
 	 */
 	UAnimNextVariableEntry* AddVariable(UAnimNextRigVMAsset* InAnimNextAsset, const FAnimNextParamType& InVariableType, const FString& InVariableName = TEXT("NewVariable"), const FString& InDefaultValue = TEXT(""));
 
+	/**
+	 * Adds a new shared variables instance to the AnimNext Asset specified.
+	 *
+	 * @param InAnimNextAsset  A pointer to the AnimNext asset instance.
+	 * @param InSharedVariables:  A pointer to the AnimNext shared variables instance.
+	 * @return A pointer to the object instance for the shared variables, or nullptr otherwise.
+	 */
+	UAnimNextSharedVariablesEntry* AddSharedVariables(UAnimNextRigVMAsset* InAnimNextAsset, UAnimNextSharedVariables* InSharedVariables);
+
 	/**
 	 * Adds a new variable node to the AnimNext graph specified.
 	 *
